Add experimental cross-platform RTSPMedia plugin

Contains an RTSP media source implementation which supports H.264 video.

Implements RTSP stream negotiation and H.264 RTP depacketization.
Leans on Electra decoders for H.264 video decoding across Windows, Mac and Linux.

#rnx
#Synced-CL 51153863

diff --git a/Engine/Plugins/Media/RTSPMedia/README.md b/Engine/Plugins/Media/RTSPMedia/README.md
new file mode 100644
--- /dev/null
+++ b/Engine/Plugins/Media/RTSPMedia/README.md
@@ -0,0 +1,36 @@
+# RTSP Media
+
+An experimental plugin for Unreal Engine's Media Framework containing a cross-platform RTSP media source implementation with H.264 support.
+
+Utilises ElectraDecoders for H.264 decoding. 
+
+## Platforms
+
+- Windows
+- macOS
+- Linux
+
+## Features
+
+- TCP interleaved transport to ensure packet delivery.
+- Configurable jitter buffer to smooth network delivery spikes. Set this to zero to for minimum latency on clean networks.
+- Automatic reconnection on connection failure.
+- Configurable decoder buffer size. Reduce to 1 for streams that don't contain H.264 B-frames to minimise latency.
+
+## Linux Setup
+
+Linux support requires an installation of libavcodec and libavutil to be available on the local system. 
+
+Instructions to set this up are available in the following location:
+
+`Engine/Source/ThirdParty/libav/README`
+
+## Current Limitations
+
+- H.264 only
+- TCP interleaved transport only
+  - No UDP support
+- Fixed configuration per session
+  - No automatic jitter buffer adjustment
+- No audio support
+- No authentication support
diff --git a/Engine/Plugins/Media/RTSPMedia/RTSPMedia.uplugin b/Engine/Plugins/Media/RTSPMedia/RTSPMedia.uplugin
new file mode 100644
--- /dev/null
+++ b/Engine/Plugins/Media/RTSPMedia/RTSPMedia.uplugin
@@ -0,0 +1,44 @@
+{
+    "FileVersion": 3,
+    "Version": 1,
+    "VersionName": "1.0",
+    "FriendlyName": "RTSP Media",
+    "Description": "Real-time media streaming via the RTSP protocol",
+    "Category": "Media Players",
+    "CreatedBy": "Epic Games",
+    "CreatedByURL": "https://epicgames.com",
+    "CanContainContent": true,
+    "IsExperimentalVersion": true,
+    "EnabledByDefault": false,
+    "SupportedTargetPlatforms": [
+        "Win64",
+        "Mac",
+        "Linux"
+    ],
+    "Modules": [
+        {
+            "Name": "RTSPMedia",
+            "Type": "Runtime",
+            "LoadingPhase": "Default"
+        },
+        {
+            "Name": "RTSPMediaEditor",
+            "Type": "Editor",
+            "LoadingPhase": "PostEngineInit"
+        }
+    ],
+    "Plugins": [
+        {
+            "Name": "ElectraCodecs",
+            "Enabled": true
+        },
+        {
+            "Name": "MediaIOFramework",
+            "Enabled": true
+        },
+        {
+            "Name": "MediaPlayerEditor",
+            "Enabled": true
+        }
+    ]
+}
\ No newline at end of file
diff --git a/Engine/Plugins/Media/RTSPMedia/Source/RTSPMedia/Private/Player/RtspMediaPlayer.cpp b/Engine/Plugins/Media/RTSPMedia/Source/RTSPMedia/Private/Player/RtspMediaPlayer.cpp
new file mode 100644
--- /dev/null
+++ b/Engine/Plugins/Media/RTSPMedia/Source/RTSPMedia/Private/Player/RtspMediaPlayer.cpp
@@ -0,0 +1,394 @@
+// Copyright Epic Games, Inc. All Rights Reserved.
+
+#include "Player/RtspMediaPlayer.h"
+
+// RTSPMedia
+#include "RTP/RtpDecoder.h"
+#include "RTP/RtpPacket.h"
+#include "RtspMediaConstants.h"
+#include "RtspMediaDefaults.h"
+#include "RtspMediaSource.h"
+#include "SDP/SdpSession.h"
+
+#include "CoreMinimal.h"
+
+// Media Framework
+#include "IMediaEventSink.h"
+#include "IMediaOptions.h"
+#include "MediaSamples.h"
+#include "Async/Async.h"
+
+FRtspMediaPlayer::FRtspMediaPlayer(IMediaEventSink& InEventSink)
+	: Super(InEventSink)
+	// We always need a Samples object, even if we replace it later
+	// FRtspMediaPlayer::GetSamples returns a reference to a FMediaSamples instance
+	, Samples(new FMediaSamples(0, 0))
+{}
+
+FRtspMediaPlayer::~FRtspMediaPlayer()
+{
+	FRtspMediaPlayer::Close();
+}
+
+// UsePlaybackTimingV2 is essential here, it dictates which overriden methods are called to trigger the gathering of samples.
+bool FRtspMediaPlayer::GetPlayerFeatureFlag(EFeatureFlag InFlag) const
+{
+	if (InFlag == EFeatureFlag::AlwaysPullNewestVideoFrame ||
+		InFlag == EFeatureFlag::AllowShutdownOnClose ||
+		InFlag == EFeatureFlag::UsePlaybackTimingV2)
+	{
+		return true;
+	}
+
+	return false;
+}
+
+FGuid FRtspMediaPlayer::GetPlayerPluginGUID() const
+{
+	return RtspMedia::PlayerPluginGUID;
+}
+
+FTimespan FRtspMediaPlayer::GetTime() const
+{
+	// When using V2 timing the media framework uses the timestamp within the media samples
+	return FTimespan::Zero();
+}
+
+bool FRtspMediaPlayer::Open(const FString& InUrl, const IMediaOptions* InOptions)
+{
+	if (!ensureMsgf(InOptions != nullptr, TEXT("Media player open was attempted but no options were provided.")))
+	{
+		return false;
+	}
+	
+	CurrentUrl = InUrl;
+	
+	UE_LOG(LogRtspMedia, Verbose, TEXT("Reading options"));
+	const int64 MaxQueuedVideoSamplesOption = InOptions->GetMediaOption(RtspMedia::Option::MaxQueuedVideoSamples, static_cast<int64>(RtspMedia::Default::MaxQueuedVideoSamples));
+	Samples = MakeUnique<FMediaSamples>(RtspMedia::Default::MaxQueuedAudioSamples, static_cast<int32>(FMath::Clamp(MaxQueuedVideoSamplesOption, 2, 16)));
+
+	const double RequestTimeoutSecondsOption = InOptions->GetMediaOption(RtspMedia::Option::RequestTimeoutSeconds, static_cast<double>(RtspMedia::Default::RequestTimeoutSeconds));
+	const int64 SocketBufferSizeKbOption = InOptions->GetMediaOption(RtspMedia::Option::SocketBufferSize, static_cast<int64>(RtspMedia::Default::SocketBufferSizeKb));
+	const int64 TransportProtocolOption = InOptions->GetMediaOption(RtspMedia::Option::TransportProtocol, static_cast<int64>(RtspMedia::Default::TransportProtocol));
+	const int64 JitterBufferDepthMsOption = InOptions->GetMediaOption(RtspMedia::Option::JitterBufferDepthMs, static_cast<int64>(RtspMedia::Default::JitterBufferDepthMs));
+
+	RtspClientConfiguration = MakeUnique<FRtspClientConfiguration>();
+	RtspClientConfiguration->Url = InUrl;
+	RtspClientConfiguration->RequestTimeoutSeconds = static_cast<float>(FMath::Clamp(RequestTimeoutSecondsOption, 1.0, 30.0));
+	RtspClientConfiguration->SocketBufferSizeBytes = static_cast<int32>(FMath::Clamp(SocketBufferSizeKbOption, 128, 4096) * 1024);
+	RtspClientConfiguration->TransportProtocol = static_cast<FRtspTransportConfiguration::TransportProtocol>(FMath::Clamp(TransportProtocolOption, 0, 1));
+	RtspClientConfiguration->JitterBufferDepthMs = static_cast<uint32>(FMath::Clamp(JitterBufferDepthMsOption, 0, 5000));
+
+	ResetRestartState();
+	bAutoRestart = InOptions->GetMediaOption(RtspMedia::Option::AutoReconnect, RtspMedia::Default::bAutoReconnect);
+	const double MinReconnectDelaySecondsOption = InOptions->GetMediaOption(RtspMedia::Option::MinReconnectDelaySeconds, static_cast<double>(RtspMedia::Default::MinReconnectDelaySeconds));
+	MinRestartDelaySeconds = static_cast<float>(FMath::Clamp(MinReconnectDelaySecondsOption, 1.0, 60.0));
+	const double MaxReconnectDelaySecondsOption = InOptions->GetMediaOption(RtspMedia::Option::MaxReconnectDelaySeconds, static_cast<double>(RtspMedia::Default::MaxReconnectDelaySeconds));
+	MaxRestartDelaySeconds = static_cast<float>(FMath::Clamp(MaxReconnectDelaySecondsOption, 1.0, 300.0));
+	const int64 MaxReconnectAttemptsOption = InOptions->GetMediaOption(RtspMedia::Option::MaxReconnectAttempts, static_cast<int64>(RtspMedia::Default::MaxReconnectAttempts));
+	MaxRestartAttempts = static_cast<int32>(FMath::Clamp(MaxReconnectAttemptsOption, 0, 1000));
+
+	// Decoder options
+	MaxFragmentBufferSizeMbOption = InOptions->GetMediaOption(RtspMedia::Option::MaxFragmentBufferSize, static_cast<int64>(RtspMedia::Default::MaxFragmentBufferSizeMb));
+	DecoderBufferSizeOption = InOptions->GetMediaOption(RtspMedia::Option::DecoderBufferSize, static_cast<int64>(RtspMedia::Default::DecoderBufferSize));
+	DecoderPollIntervalMsOption = InOptions->GetMediaOption(RtspMedia::Option::DecoderPollIntervalMs, static_cast<int64>(RtspMedia::Default::DecoderPollIntervalMs));
+
+	ResetClient();
+	ResetDecoder();
+
+	// InitializeClient() may fail synchronously (e.g., Connect() fails immediately), but we still return true at the end of this method to indicate we've accepted the URL.
+	// Returning false would cause the Media Framework to abandon this player, preventing auto-reconnect from taking place via TickInput().
+	InitializeClient();
+	
+	return true;
+}
+
+void FRtspMediaPlayer::TickInput(FTimespan InDeltaTime, FTimespan InTimecode)
+{
+	if (bAutoRestart)
+	{
+		if (RestartCountdownSeconds > 0.0f)
+		{
+			RestartCountdownSeconds -= static_cast<float>(InDeltaTime.GetTotalSeconds());
+
+			if (RestartCountdownSeconds <= 0.0f)
+			{
+				RestartCountdownSeconds = 0.0f;
+				ResetClient();
+				InitializeClient();
+			}
+		}
+	}
+	
+	FStreamingMediaPlayerBase::TickInput(InDeltaTime, InTimecode);
+}
+
+void FRtspMediaPlayer::Close()
+{
+	if (CurrentState == EMediaState::Closed)
+	{
+		return;
+	}
+	
+	UE_LOG(LogRtspMedia, Log, TEXT("Closing media player"));
+	
+	ResetRestartState();
+
+	CurrentState = EMediaState::Closed;
+	EventSink.ReceiveMediaEvent(EMediaEvent::MediaClosed);
+
+	if (Samples.IsValid())
+	{
+		// Don't reset Samples, the media player always needs a valid reference until the player is destroyed
+		Samples->FlushSamples();
+	}
+
+	ResetClient();
+	ResetDecoder();
+}
+
+IMediaSamples& FRtspMediaPlayer::GetSamples()
+{
+	return *Samples;
+}
+
+bool FRtspMediaPlayer::GetVideoTrackFormat(int32 InTrackIndex, int32 InFormatIndex, FMediaVideoTrackFormat& OutFormat) const
+{
+	if (InTrackIndex == 0 && InFormatIndex == 0 && NumVideoTracks > 0)
+	{
+		const float FrameRate = RtpDecoder.IsValid() ? static_cast<float>(RtpDecoder->GetFrameRate()) : 30.0f;
+		OutFormat.Dim = FIntPoint(VideoFormatWidth, VideoFormatHeight);
+		OutFormat.FrameRate = FrameRate;
+		OutFormat.FrameRates = TRange<float>(FrameRate);
+		OutFormat.TypeName = TEXT("H.264");
+		return true;
+	}
+
+	return false;
+}
+
+void FRtspMediaPlayer::HandleOnClientStateChanged(ERtspClientState InNewState, ERtspClientState InOldState)
+{
+	switch (InNewState)
+	{
+		case ERtspClientState::Playing:
+			if (!RtpDecoder.IsValid() || !RtpDecoder->Start())
+			{
+				UE_LOG(LogRtspMedia, Error, TEXT("Failed to start RTP decoder"));
+				SetState(EMediaState::Error);
+				EventSink.ReceiveMediaEvent(EMediaEvent::MediaOpenFailed);
+				// If the decoder fails to start we intentionally close the media source rather than trigger a retry.
+				Close();
+			}
+			else
+			{
+				CurrentStatus = EMediaStatus::None;
+				CurrentState = EMediaState::Playing;
+				EventSink.ReceiveMediaEvent(EMediaEvent::PlaybackResumed);
+				ResetRestartState();
+			}
+			break;
+		case ERtspClientState::Setup:
+			EventSink.ReceiveMediaEvent(EMediaEvent::MediaOpened);
+			break;
+		case ERtspClientState::Error:
+			HandlePlayerError();
+			break;
+		default:
+			break;
+	}
+}
+
+// This method is called by RtspClient on its dedicated worker thread to avoid a game thread hop.
+void FRtspMediaPlayer::HandleOnRtpPacketReceived(FRtpPacket InPacket)
+{
+	if (!RtpDecoder.IsValid())
+	{
+		UE_LOG(LogRtspMedia, Verbose, TEXT("RTP decoder was not available when packet was received"));
+		return;
+	}
+	
+	RtpDecoder->EnqueuePacket(MoveTemp(InPacket));
+}
+
+void FRtspMediaPlayer::HandleOnSdpSessionReceived(const FSdpSession& Session)
+{
+	UE_LOG(LogRtspMedia, Verbose, TEXT("SDP received"));
+
+	if (!Session.HasVideo())
+	{
+		UE_LOG(LogRtspMedia, Error, TEXT("No video track present within SDP session"));
+		HandlePlayerError();
+		return;
+	}
+
+	// We only support a single video track at this time
+	NumVideoTracks = 1;
+
+	// Ensure the media framework is informed now we have a video track available.
+	EventSink.ReceiveMediaEvent(EMediaEvent::TracksChanged);
+	
+	// TODO: Signal audio track availability once audio is implemented
+
+	// Reset decoder in case we already have one at this point
+	ResetDecoder();
+	
+	// Create the decoder
+	const int64 MaxFragmentBufferSize = FMath::Clamp(MaxFragmentBufferSizeMbOption, 2, 64) * 1024 * 1024;
+	const int64 DecoderBufferSize = FMath::Clamp(DecoderBufferSizeOption, 1, 16);
+	const uint32 DecoderPollIntervalMs = static_cast<uint32>(FMath::Clamp(DecoderPollIntervalMsOption, 1, 10));
+	RtpDecoder = MakeUnique<FRtpDecoder>();
+	if (!RtpDecoder->Initialize(
+		Samples.Get(),
+		Session,
+		MaxFragmentBufferSize,
+		DecoderBufferSize,
+		DecoderPollIntervalMs))
+	{
+		UE_LOG(LogRtspMedia, Error, TEXT("Failed to initialise decoder"));
+		SetState(EMediaState::Error);
+		EventSink.ReceiveMediaEvent(EMediaEvent::MediaOpenFailed);
+		Close();
+	}
+	else
+	{
+		// Bind to decoder for video dimensions information and for error handling
+		RtpDecoder->OnVideoDimensions.BindSP(this, &FRtspMediaPlayer::HandleOnVideoDimensions);
+		RtpDecoder->OnDecoderError.BindSP(this, &FRtspMediaPlayer::HandleOnDecoderError);
+	}
+	
+}
+
+void FRtspMediaPlayer::HandleOnVideoDimensions(int32 InWidth, int32 InHeight)
+{
+	UE_LOG(LogRtspMedia, Verbose, TEXT("Player was provided with updated video dimensions. Width: %d Height: %d"), InWidth, InHeight);
+
+	TWeakPtr<FRtspMediaPlayer> WeakSelf = AsWeak();
+	AsyncTask(ENamedThreads::GameThread, [WeakSelf, InWidth, InHeight]
+	{
+		const TSharedPtr<FRtspMediaPlayer> Player = WeakSelf.Pin();
+		if (Player.IsValid())
+		{
+			Player->VideoFormatWidth = InWidth;
+			Player->VideoFormatHeight = InHeight;
+		}
+	});
+}
+
+void FRtspMediaPlayer::HandleOnDecoderError()
+{
+	TWeakPtr<FRtspMediaPlayer> WeakSelf = AsWeak();
+	AsyncTask(ENamedThreads::GameThread, [WeakSelf]
+	{
+		const TSharedPtr<FRtspMediaPlayer> Player = WeakSelf.Pin();
+		if (Player.IsValid())
+		{
+			Player->HandlePlayerError();
+		}
+	});
+}
+
+void FRtspMediaPlayer::ResetClient()
+{
+	if (RtspClient.IsValid())
+	{
+		RtspClient->OnStateChanged.Unbind();
+		RtspClient->OnSdpSessionReceived.Unbind();
+		RtspClient->OnRtpPacketReceived.Unbind();
+
+		RtspClient->Disconnect();
+		RtspClient.Reset();
+	}
+}
+
+void FRtspMediaPlayer::ResetDecoder()
+{
+	if (RtpDecoder.IsValid())
+	{
+		RtpDecoder->OnDecoderError.Unbind();
+		RtpDecoder->OnVideoDimensions.Unbind();
+		RtpDecoder->Shutdown();
+		RtpDecoder.Reset();
+	}
+}
+
+void FRtspMediaPlayer::ResetRestartState()
+{
+	RestartAttempts = 0;
+	RestartCountdownSeconds = 0.0f;
+}
+
+void FRtspMediaPlayer::InitializeClient()
+{
+	UE_LOG(LogRtspMedia, Verbose, TEXT("Initialising RTSP Client"));
+
+	RtspClient = MakeUnique<FRtspClient>();
+
+	RtspClient->OnSdpSessionReceived.BindSP(this, &FRtspMediaPlayer::HandleOnSdpSessionReceived);
+	RtspClient->OnStateChanged.BindSP(this, &FRtspMediaPlayer::HandleOnClientStateChanged);
+	RtspClient->OnRtpPacketReceived.BindSP(this, &FRtspMediaPlayer::HandleOnRtpPacketReceived);
+
+	if (RtspClient->Initialize(*RtspClientConfiguration))
+	{
+		CurrentStatus = EMediaStatus::Connecting;
+		CurrentState = EMediaState::Preparing;
+		EventSink.ReceiveMediaEvent(EMediaEvent::MediaConnecting);
+	}
+	else
+	{
+		CurrentStatus = EMediaStatus::None;
+		CurrentState = EMediaState::Error;
+		EventSink.ReceiveMediaEvent(EMediaEvent::MediaOpenFailed);
+	}
+}
+
+void FRtspMediaPlayer::HandlePlayerError()
+{
+	if (CurrentState == EMediaState::Closed)
+	{
+		return;
+	}
+	
+	ResetClient();
+	
+	if (bAutoRestart && (MaxRestartAttempts == 0 || RestartAttempts < MaxRestartAttempts))
+	{
+		ScheduleRestart();
+		return;
+	}
+
+	if (RestartAttempts >= MaxRestartAttempts)
+	{
+		UE_LOG(LogRtspMedia, Warning, TEXT("Reached the max number of client restart attempts (%d). Stopping."), MaxRestartAttempts);
+	}
+	
+	CurrentStatus = EMediaStatus::None;
+	CurrentState = EMediaState::Error;
+	EventSink.ReceiveMediaEvent(EMediaEvent::MediaOpenFailed);
+}
+
+// Calculate how long we should wait until the next retry attempt
+// We cap the retry period at MaxRestartDelaySeconds and introduce
+// a small random offset in case multiple clients all lose their
+// connection at the same time to avoid spamming the server simultaneously. 
+void FRtspMediaPlayer::ScheduleRestart()
+{
+	// How quickly do we want to reach the maximum restart delay?
+	// If we have a limited number of retry attempts, use that value,
+	// otherwise we'll hit the max retry period after 10.0 attempts.
+	const float ScaleFactor = MaxRestartAttempts > 0 ? static_cast<float>(MaxRestartAttempts) : 10.0f;
+	// Normalised progress towards max retry delay
+	const float Progress = FMath::Min(static_cast<float>(RestartAttempts) / ScaleFactor, 1.0f);
+	// Cubic interpolation so that the first few delays stay low then ramp up 
+	const float Delay = FMath::InterpEaseIn(MinRestartDelaySeconds, MaxRestartDelaySeconds, Progress, 3.0f);
+
+	// Applying a small random offset to the delay so that multiple sources don't all hammer the server at once 
+	RestartCountdownSeconds = Delay + FMath::FRandRange(0.0f, 1.0f);
+	RestartAttempts++;
+
+	UE_LOG(LogRtspMedia, Log, TEXT("Restarting RTSP client in %.1f seconds. Attempt number %d"), RestartCountdownSeconds, RestartAttempts);
+
+	CurrentStatus = EMediaStatus::Connecting;
+	CurrentState = EMediaState::Preparing;
+	EventSink.ReceiveMediaEvent(EMediaEvent::MediaConnecting);
+}
diff --git a/Engine/Plugins/Media/RTSPMedia/Source/RTSPMedia/Private/Player/RtspMediaPlayer.h b/Engine/Plugins/Media/RTSPMedia/Source/RTSPMedia/Private/Player/RtspMediaPlayer.h
new file mode 100644
--- /dev/null
+++ b/Engine/Plugins/Media/RTSPMedia/Source/RTSPMedia/Private/Player/RtspMediaPlayer.h
@@ -0,0 +1,80 @@
+// Copyright Epic Games, Inc. All Rights Reserved.
+
+#pragma once
+
+#include "Player/StreamingMediaPlayerBase.h"
+
+#include "RTSP/RtspClient.h"
+
+#include "RTP/RtpDecoder.h"
+#include "Templates/SharedPointer.h"
+
+class FRtpDecoder;
+struct FRtpPacket;
+
+class FMediaSamples;
+
+class FRtspMediaPlayer 
+	: public FStreamingMediaPlayerBase
+	, public TSharedFromThis<FRtspMediaPlayer, ESPMode::ThreadSafe>
+{
+	using Super = FStreamingMediaPlayerBase;
+
+public:
+	FRtspMediaPlayer(IMediaEventSink& InEventSink);
+	virtual ~FRtspMediaPlayer() override;
+
+	//~ Begin IMediaPlayer
+	virtual void Close() override;
+	virtual IMediaSamples& GetSamples() override;
+	virtual bool GetPlayerFeatureFlag(EFeatureFlag InFlag) const override;
+	virtual FGuid GetPlayerPluginGUID() const override;
+	virtual FTimespan GetTime() const override;
+	virtual bool Open(const FString& InUrl, const IMediaOptions* InOptions) override;
+	virtual void TickInput(FTimespan InDeltaTime, FTimespan InTimecode) override;
+	//~ End IMediaPlayer
+
+	// Overridden from FStreamingMediaPlayerBase
+	virtual bool GetVideoTrackFormat(int32 InTrackIndex, int32 InFormatIndex, FMediaVideoTrackFormat& OutFormat) const override;
+
+private:
+	
+	void HandleOnClientStateChanged(ERtspClientState InNewState, ERtspClientState InOldState);
+	void HandleOnRtpPacketReceived(FRtpPacket InPacket);
+	void HandleOnSdpSessionReceived(const FSdpSession& InSession);
+	void HandleOnVideoDimensions(int32 InWidth, int32 InHeight);
+	void HandleOnDecoderError();
+
+	void ResetClient();
+	void ResetDecoder();
+	void ResetRestartState();
+
+	void InitializeClient();
+	void HandlePlayerError();
+	void ScheduleRestart();
+	
+	// RTSP Networking
+	TUniquePtr<FRtspClientConfiguration> RtspClientConfiguration;
+	TUniquePtr<FRtspClient> RtspClient;
+
+	// Decoding
+	int64 MaxFragmentBufferSizeMbOption = RtspMedia::Default::MaxFragmentBufferSizeMb;
+	int64 DecoderBufferSizeOption = RtspMedia::Default::DecoderBufferSize;
+	int64 DecoderPollIntervalMsOption = RtspMedia::Default::DecoderPollIntervalMs;
+	TUniquePtr<FRtpDecoder> RtpDecoder;
+	
+	// Media Player
+	TUniquePtr<FMediaSamples> Samples;
+	int32 VideoFormatWidth = 0;
+	int32 VideoFormatHeight = 0;
+
+	// Auto Restart Configuration
+	bool bAutoRestart = RtspMedia::Default::bAutoReconnect;
+	float MinRestartDelaySeconds = RtspMedia::Default::MinReconnectDelaySeconds;
+	float MaxRestartDelaySeconds = RtspMedia::Default::MaxReconnectDelaySeconds;
+	int32 MaxRestartAttempts = RtspMedia::Default::MaxReconnectAttempts;
+
+	// Restart State
+	int32 RestartAttempts = 0;
+	float RestartCountdownSeconds = 0.0f;
+};
diff --git a/Engine/Plugins/Media/RTSPMedia/Source/RTSPMedia/Private/Player/StreamingMediaPlayerBase.cpp b/Engine/Plugins/Media/RTSPMedia/Source/RTSPMedia/Private/Player/StreamingMediaPlayerBase.cpp
new file mode 100644
--- /dev/null
+++ b/Engine/Plugins/Media/RTSPMedia/Source/RTSPMedia/Private/Player/StreamingMediaPlayerBase.cpp
@@ -0,0 +1,275 @@
+// Copyright Epic Games, Inc. All Rights Reserved.
+
+#include "Player/StreamingMediaPlayerBase.h"
+
+#include "IMediaEventSink.h"
+
+#define LOCTEXT_NAMESPACE "StreamMediaPlayerBase"
+
+FStreamingMediaPlayerBase::FStreamingMediaPlayerBase(IMediaEventSink& InEventSink)
+	: EventSink(InEventSink)
+{}
+
+void FStreamingMediaPlayerBase::SetState(EMediaState InNewState)
+{
+	if (CurrentState != InNewState)
+	{
+		CurrentState = InNewState;
+	}
+}
+
+//~ Begin IMediaPlayer
+
+IMediaCache& FStreamingMediaPlayerBase::GetCache()
+{
+	return *this;
+}
+
+IMediaControls& FStreamingMediaPlayerBase::GetControls()
+{
+	return *this;
+}
+
+FString FStreamingMediaPlayerBase::GetInfo() const
+{
+	// Subclass could override to display relevant debug info
+	return FString();
+}
+
+FString FStreamingMediaPlayerBase::GetStats() const
+{
+	// Subclass could override to display stats
+	return FString();
+}
+
+IMediaTracks& FStreamingMediaPlayerBase::GetTracks()
+{
+	return *this;
+}
+
+FString FStreamingMediaPlayerBase::GetUrl() const
+{
+	return CurrentUrl;
+}
+
+IMediaView& FStreamingMediaPlayerBase::GetView()
+{
+	return *this;
+}
+
+bool FStreamingMediaPlayerBase::Open(const TSharedRef<FArchive, ESPMode::ThreadSafe>& InArchive, const FString& InOriginalUrl, const IMediaOptions* InOptions)
+{
+	// We don't support opening from archives.
+	return false;
+}
+
+//~ End IMediaPlayer
+
+//~ Begin IMediaControls
+
+bool FStreamingMediaPlayerBase::CanControl(EMediaControl InControl) const 
+{
+	// We won't offer pause, seeking or scrub for live streams
+	return false;
+}
+
+FTimespan FStreamingMediaPlayerBase::GetDuration() const 
+{
+	return FTimespan::MaxValue();
+}
+
+float FStreamingMediaPlayerBase::GetRate() const 
+{
+	// Normal speed
+	return 1.0f;
+}
+
+EMediaState FStreamingMediaPlayerBase::GetState() const 
+{
+	return CurrentState;
+}
+
+EMediaStatus FStreamingMediaPlayerBase::GetStatus() const 
+{
+	// Subclasses can override to provide status information
+	return CurrentStatus;
+}
+
+TRangeSet<float> FStreamingMediaPlayerBase::GetSupportedRates(EMediaRateThinning InThinning) const 
+{
+	// We only support a single rate, normal speed.
+	TRangeSet<float> Rates;
+	Rates.Add(TRange<float>(1.0f));
+	return Rates;
+}
+
+bool FStreamingMediaPlayerBase::IsLooping() const 
+{
+	return false;
+}
+
+bool FStreamingMediaPlayerBase::Seek(const FTimespan& InTime) 
+{
+	return false;
+}
+
+bool FStreamingMediaPlayerBase::SetLooping(bool bInLooping) 
+{
+	return false;
+}
+
+bool FStreamingMediaPlayerBase::SetRate(float InRate) 
+{
+	// Only normal speed (1.0) is accepted
+	return FMath::IsNearlyEqual(InRate, 1.0f);
+}
+
+//~ End IMediaControls
+
+//~ Begin IMediaTracks
+
+bool FStreamingMediaPlayerBase::GetAudioTrackFormat(int32 InTrackIndex, int32 InFormatIndex, FMediaAudioTrackFormat& OutFormat) const
+{
+	// Subclasses override to provide track format information
+	return false;
+}
+
+int32 FStreamingMediaPlayerBase::GetNumTracks(EMediaTrackType InTrackType) const
+{
+	switch (InTrackType)
+	{
+		case EMediaTrackType::Video: return NumVideoTracks;
+		case EMediaTrackType::Audio: return NumAudioTracks;
+		default: return 0;
+	}
+}
+
+int32 FStreamingMediaPlayerBase::GetNumTrackFormats(EMediaTrackType InTrackType, int32 InTrackIndex) const
+{
+	// We only support a single format for the track types we support
+	if (InTrackIndex != 0)
+	{
+		return 0;
+	}
+	switch (InTrackType)
+	{
+		case EMediaTrackType::Video: return NumVideoTracks > 0 ? 1 : 0;
+		case EMediaTrackType::Audio: return NumAudioTracks > 0 ? 1 : 0;
+		default: return 0;
+	}
+}
+
+int32 FStreamingMediaPlayerBase::GetSelectedTrack(EMediaTrackType InTrackType) const
+{
+	switch (InTrackType)
+	{
+		case EMediaTrackType::Video: return SelectedVideoTrack;
+		case EMediaTrackType::Audio: return SelectedAudioTrack;
+		default: return INDEX_NONE;
+	}
+}
+
+FText FStreamingMediaPlayerBase::GetTrackDisplayName(EMediaTrackType InTrackType, int32 InTrackIndex) const
+{
+	switch (InTrackType)
+	{
+		case EMediaTrackType::Video: return NumVideoTracks > 0 ? LOCTEXT("VideoTrackDisplayName", "Video") : FText();
+		case EMediaTrackType::Audio: return NumAudioTracks > 0 ? LOCTEXT("AudioTrackDisplayName", "Audio") : FText();
+		default: return FText();
+	}
+}
+
+int32 FStreamingMediaPlayerBase::GetTrackFormat(EMediaTrackType InTrackType, int32 InTrackIndex) const
+{
+	if (InTrackIndex == 0)
+	{
+		if (InTrackType == EMediaTrackType::Video && NumVideoTracks > 0)
+		{
+			return 0;
+		}
+		else if (InTrackType == EMediaTrackType::Audio && NumAudioTracks > 0)
+		{
+			return 0;
+		}
+	}
+	return INDEX_NONE;
+}
+
+FString FStreamingMediaPlayerBase::GetTrackLanguage(EMediaTrackType InTrackType, int32 InTrackIndex) const
+{
+	return FString();
+}
+
+FString FStreamingMediaPlayerBase::GetTrackName(EMediaTrackType InTrackType, int32 InTrackIndex) const
+{
+	switch (InTrackType)
+	{
+		case EMediaTrackType::Video: return NumVideoTracks > 0 ? TEXT("video") : FString();
+		case EMediaTrackType::Audio: return NumAudioTracks > 0 ? TEXT("audio") : FString();
+		default: return FString();
+	}
+}
+
+bool FStreamingMediaPlayerBase::GetVideoTrackFormat(int32 InTrackIndex, int32 InFormatIndex, FMediaVideoTrackFormat& OutFormat) const
+{
+	// Subclasses override to provide video track information
+	return false;
+}
+
+bool FStreamingMediaPlayerBase::SelectTrack(EMediaTrackType InTrackType, int32 InTrackIndex)
+{
+	if (InTrackIndex < 0 && InTrackIndex != INDEX_NONE)
+	{
+		return false;
+	}
+	
+	switch (InTrackType)
+	{
+		case EMediaTrackType::Video:
+			if (InTrackIndex >= NumVideoTracks)
+			{
+				return false;
+			}
+			SelectedVideoTrack = InTrackIndex;
+			return true;
+		case EMediaTrackType::Audio: 
+			if (InTrackIndex >= NumAudioTracks)
+			{
+				return false;
+			}
+			SelectedAudioTrack = InTrackIndex;
+			return true;
+		default: 
+			return false;
+	}
+}
+
+bool FStreamingMediaPlayerBase::SetTrackFormat(EMediaTrackType InTrackType, int32 InTrackIndex, int32 InFormatIndex)
+{
+	// We only support a single format for the track types we support
+	if (InFormatIndex != 0)
+	{
+		return false;
+	}
+
+	// We only support one track of each type
+	switch (InTrackType)
+	{
+		case EMediaTrackType::Video:
+			return InTrackIndex == 0 && NumVideoTracks > 0;
+		case EMediaTrackType::Audio:
+			return InTrackIndex == 0 && NumAudioTracks > 0;
+		default:
+			return false;
+	}
+}
+
+bool FStreamingMediaPlayerBase::SetVideoTrackFrameRate(int32 InTrackIndex, int32 InFormatIndex, float InFrameRate)
+{
+	// Can't control the frame rate from the client by default
+	return false;
+}
+
+//~ End IMediaTracks
+
+#undef LOCTEXT_NAMESPACE
diff --git a/Engine/Plugins/Media/RTSPMedia/Source/RTSPMedia/Private/Player/StreamingMediaPlayerBase.h b/Engine/Plugins/Media/RTSPMedia/Source/RTSPMedia/Private/Player/StreamingMediaPlayerBase.h
new file mode 100644
--- /dev/null
+++ b/Engine/Plugins/Media/RTSPMedia/Source/RTSPMedia/Private/Player/StreamingMediaPlayerBase.h
@@ -0,0 +1,80 @@
+// Copyright Epic Games, Inc. All Rights Reserved.
+
+#pragma once
+
+#include "IMediaPlayer.h"
+#include "IMediaControls.h"
+#include "IMediaTracks.h"
+#include "IMediaCache.h"
+#include "IMediaView.h"
+
+class IMediaEventSink;
+
+class FStreamingMediaPlayerBase
+	: public IMediaPlayer
+	, public IMediaControls
+	, public IMediaTracks
+	, public IMediaCache
+	, public IMediaView
+{
+public:
+	FStreamingMediaPlayerBase(IMediaEventSink &MediaEventSink);
+	virtual ~FStreamingMediaPlayerBase() = default;
+
+	virtual void Close() = 0;
+	virtual FGuid GetPlayerPluginGUID() const = 0;
+	virtual IMediaSamples& GetSamples() = 0;
+	virtual bool Open(const FString& InUrl, const IMediaOptions* InOptions) = 0;
+
+	//~ Begin IMediaPlayer
+	IMediaCache& GetCache() override;
+	IMediaControls& GetControls() override;
+	FString GetInfo() const override;
+	FString GetStats() const override;
+	IMediaTracks& GetTracks() override;
+	FString GetUrl() const override;
+	IMediaView& GetView() override;
+	bool Open(const TSharedRef<FArchive, ESPMode::ThreadSafe>& InArchive, const FString& InOriginalUrl, const IMediaOptions* InOptions) override;
+	//~ End IMediaPlayer
+
+	//~ Begin IMediaControls
+	bool CanControl(EMediaControl InControl) const override;
+	FTimespan GetDuration() const override;
+	float GetRate() const override;
+	EMediaState GetState() const override;
+	EMediaStatus GetStatus() const override;
+	TRangeSet<float> GetSupportedRates(EMediaRateThinning InThinning) const override;
+	bool IsLooping() const override;
+	bool Seek(const FTimespan& InTime) override;
+	bool SetLooping(bool bInLooping) override;
+	bool SetRate(float InRate) override;
+	//~ End IMediaControls
+
+	//~ Begin IMediaTracks
+	bool GetAudioTrackFormat(int32 InTrackIndex, int32 InFormatIndex, FMediaAudioTrackFormat& OutFormat) const override;
+	int32 GetNumTracks(EMediaTrackType InTrackType) const override;
+	int32 GetNumTrackFormats(EMediaTrackType InTrackType, int32 InTrackIndex) const override;
+	int32 GetSelectedTrack(EMediaTrackType InTrackType) const override;
+	FText GetTrackDisplayName(EMediaTrackType InTrackType, int32 InTrackIndex) const override;
+	int32 GetTrackFormat(EMediaTrackType InTrackType, int32 InTrackIndex) const override;
+	FString GetTrackLanguage(EMediaTrackType InTrackType, int32 InTrackIndex) const override;
+	FString GetTrackName(EMediaTrackType InTrackType, int32 InTrackIndex) const override;
+	bool GetVideoTrackFormat(int32 InTrackIndex, int32 InFormatIndex, FMediaVideoTrackFormat& OutFormat) const override;
+	bool SelectTrack(EMediaTrackType InTrackType, int32 InTrackIndex) override;
+	bool SetTrackFormat(EMediaTrackType InTrackType, int32 InTrackIndex, int32 InFormatIndex) override;
+	bool SetVideoTrackFrameRate(int32 InTrackIndex, int32 InFormatIndex, float InFrameRate) override;
+	//~ End IMediaTracks
+
+protected:
+
+	void SetState(EMediaState InNewState);
+
+	FString CurrentUrl;
+	EMediaState CurrentState = EMediaState::Closed;
+	EMediaStatus CurrentStatus = EMediaStatus::None;
+	IMediaEventSink& EventSink;
+	int32 NumAudioTracks = 0;
+	int32 NumVideoTracks = 0;
+	int32 SelectedVideoTrack = INDEX_NONE;
+	int32 SelectedAudioTrack = INDEX_NONE;
+};
\ No newline at end of file
diff --git a/Engine/Plugins/Media/RTSPMedia/Source/RTSPMedia/Private/RTP/RtpDecoder.cpp b/Engine/Plugins/Media/RTSPMedia/Source/RTSPMedia/Private/RTP/RtpDecoder.cpp
new file mode 100644
--- /dev/null
+++ b/Engine/Plugins/Media/RTSPMedia/Source/RTSPMedia/Private/RTP/RtpDecoder.cpp
@@ -0,0 +1,610 @@
+// Copyright Epic Games, Inc. All Rights Reserved.
+
+#include "RtpDecoder.h"
+
+#include "RtpH264Depacketizer.h"
+
+#include "RtspMediaConstants.h"
+
+#include "SDP/SdpSession.h"
+
+// Electra
+#include "ElectraTextureSample.h"
+#include "IElectraDecoder.h"
+#include "IElectraCodecFactory.h"
+#include "IElectraCodecFactoryModule.h"
+#include "IElectraDecoderOutputVideo.h"
+#include "MediaDecoderOutput.h"
+#include "ParameterDictionary.h"
+#include "Async/Async.h"
+#include "Utils/MPEG/ElectraUtilsMPEGVideo_H264.h"
+
+#include "HAL/Event.h"
+#include "HAL/PlatformProcess.h"
+#include "Misc/ScopeExit.h"
+
+FRtpDecoder::FRtpDecoder()
+	: bStopping(false)
+	, FrameRate(0.0)
+{
+	WaitEvent = FPlatformProcess::GetSynchEventFromPool();
+}
+
+FRtpDecoder::~FRtpDecoder()
+{
+	bStopping.store(true);
+
+	if (WaitEvent)
+	{
+		WaitEvent->Trigger();
+	}
+	
+	if (Thread.IsValid())
+	{
+		Thread->WaitForCompletion();
+	}
+
+	// Must wait for the thread to complete before we can return the event to the pool
+	if (WaitEvent)
+	{
+		FPlatformProcess::ReturnSynchEventToPool(WaitEvent);
+		WaitEvent = nullptr;
+	}
+
+	ResetPlatformVideoResource();
+	TextureSamplePool.Reset();
+	
+	if (H264Depacketizer.IsValid())
+	{
+		H264Depacketizer->OnH264NalUnit.Unbind();
+		H264Depacketizer.Reset();
+	}
+
+	if (ElectraDecoder.IsValid())
+	{
+		ElectraDecoder->Close();
+	}
+}
+
+bool FRtpDecoder::Initialize(FMediaSamples* InSamples, const FSdpSession& InSession, int64 InMaxFragmentBufferSize, int64 InDecoderBufferSize, uint32 InDecoderPollIntervalMs)
+{
+	Samples = InSamples;
+	MaxFragmentBufferSize = InMaxFragmentBufferSize;
+	DecoderBufferSize = InDecoderBufferSize;
+	DecoderPollIntervalMs = InDecoderPollIntervalMs;
+	TextureSamplePool = MakeUnique<FElectraTextureSamplePool>();
+	
+	SdpSession = MakeUnique<FSdpSession>(InSession);
+	
+	if (!SdpSession.IsValid())
+	{
+		UE_LOG(LogRtspMedia, Error, TEXT("Session unavailable when creating electra decoder"));
+		return false;
+	}
+
+	if (!SdpSession->HasVideo())
+	{
+		UE_LOG(LogRtspMedia, Error, TEXT("Session must contain at least one video track"));
+		return false;
+	}
+
+	const FSdpMediaTrack* VideoTrack = SdpSession->GetFirstVideoTrack();
+	if (!VideoTrack)
+	{
+		UE_LOG(LogRtspMedia, Error, TEXT("Session did not provide a video track"));
+		return false;
+	}
+
+	SequenceParameterSet = VideoTrack->SequenceParameterSet;
+	PictureParameterSet = VideoTrack->PictureParameterSet;
+
+	// Store the clock rate to construct presentation time stamps when decoding
+	const uint32 ClockRate = VideoTrack->ClockRate;
+	if (ClockRate == 0)
+	{
+		UE_LOG(LogRtspMedia, Error, TEXT("Session contains a zero clock rate when initializing decoder"));
+		return false;
+	}
+	VideoClockRate = ClockRate;
+	
+	return CreateElectraDecoder(SequenceParameterSet, PictureParameterSet);
+}
+
+bool FRtpDecoder::Start()
+{
+	if (!ensureMsgf(ElectraDecoder.IsValid(), TEXT("Decoder cannot be started if electra decoder is not initialised")))
+	{
+		return false;
+	}
+
+	bStopping.store(false);
+	
+	Thread.Reset(FRunnableThread::Create(this, TEXT("RtpDecoderThread"), 0, TPri_Normal));
+
+	if (!Thread.IsValid())
+	{
+		UE_LOG(LogRtspMedia, Error, TEXT("Failed to create decoder thread"));
+		return false;
+	}
+
+	return true;
+}
+
+void FRtpDecoder::Shutdown()
+{
+	Stop();
+}
+
+void FRtpDecoder::EnqueuePacket(FRtpPacket InPacket)
+{
+	Queue.Enqueue(MoveTemp(InPacket));
+	if (WaitEvent)
+	{
+		WaitEvent->Trigger();
+	}
+}
+
+double FRtpDecoder::GetFrameRate() const
+{
+	return FrameRate.load(std::memory_order_relaxed);
+}
+
+bool FRtpDecoder::Init()
+{
+	return true;
+}
+
+uint32 FRtpDecoder::Run()
+{
+	while (!bStopping.load())
+	{
+		// Check for any decoded output in the electra decoder
+		PollDecodedOutput();
+
+		// Delay until a packet arrives or the decoder poll interval expires
+		if (WaitEvent && WaitEvent->Wait(DecoderPollIntervalMs))
+		{
+			FRtpPacket Packet;
+			while (Queue.Dequeue(Packet))
+			{
+				if (!H264Depacketizer.IsValid())
+				{
+					H264Depacketizer = MakeUnique<FRtpH264Depacketizer>();
+					H264Depacketizer->SetMaxFragmentBufferSize(MaxFragmentBufferSize);
+					H264Depacketizer->OnH264NalUnit.BindRaw(this, &FRtpDecoder::HandleOnH264NalUnit);
+				}
+				const FRtpHeader Header = Packet.Header;
+				H264Depacketizer->ProcessRtpPayload(MoveTemp(Packet.Payload), Header.SequenceNumber, Header.Timestamp, Header.bMarker);
+			}
+		}
+	}
+
+	return 0;
+}
+
+void FRtpDecoder::Stop()
+{
+	bStopping.store(true);
+	if (WaitEvent)
+	{
+		WaitEvent->Trigger();
+	}
+}
+
+ bool FRtpDecoder::CreateElectraDecoder(const TArray<uint8>& InSequenceParameterSet, const TArray<uint8>& InPictureParameterSet)
+ {
+	// Build Codec Specific Data ('csd')
+	// Used to initalize the decoder and for each decode action.
+	// Contains the sequence and picture parameter sets which include information like:
+	// SPS - Properties for the entire sequence of frames
+	// E.g. Resolution, profile, chroma format, bit depth
+	// PPS - Per-frame properties
+	// E.g. Coding mode, quantization parameters
+
+	// H.264 Annex B start code is required before the SPS and PPS
+	constexpr uint8 H264AnnexBStartCode[] = {0x00, 0x00, 0x00, 0x01};
+
+	TArray<uint8> CodecSpecificData;
+	CodecSpecificData.Append(H264AnnexBStartCode, 4);
+	CodecSpecificData.Append(InSequenceParameterSet);
+	CodecSpecificData.Append(H264AnnexBStartCode, 4);
+	CodecSpecificData.Append(InPictureParameterSet);
+
+	CodecOptions.Empty();
+
+	FString CodecFormat;
+
+	// Convert Annex B style codec specific data to an avcC configuration record
+	ElectraDecodersUtil::MPEG::H264::FAVCDecoderConfigurationRecord AvcConfigurationRecord;
+	if (AvcConfigurationRecord.CreateFromCodecSpecificData(CodecSpecificData))
+	{
+		// DCR = Decoder Configuration Record
+		CodecOptions.Emplace("dcr", AvcConfigurationRecord.GetRawData());
+		CodecOptions.Emplace("csd", AvcConfigurationRecord.GetCodecSpecificData());
+		CodecOptions.Emplace("max_output_buffers", FVariant(DecoderBufferSize));
+
+		// 'avc1' is a flavour of H.264 where the Sequence and Picture parameter sets are stored 'out of band'
+		// This means they are stored separately from the main H.264 bitstream, which is the case when streaming with RTSP.
+		// In RTSP these parameter sets are provided within the SDP (Session Description Protocol) data when the connection handshake occurs.
+		CodecFormat = AvcConfigurationRecord.GetCodecSpecifierRFC6381(TEXT("avc1"));
+
+		const TArray<TArray<uint8>>& SequenceParameterSets = AvcConfigurationRecord.GetSequenceParameterSets();
+		if (!SequenceParameterSets.IsEmpty())
+		{
+			int32 MaxWidth = 0;
+			int32 MaxHeight = 0;
+			for (const TArray<uint8>& CurrentSequenceParameterSet : SequenceParameterSets)
+			{
+				ElectraDecodersUtil::MPEG::H264::FSequenceParameterSet ParsedSequenceParameterSet;
+				if (ElectraDecodersUtil::MPEG::H264::ParseSequenceParameterSet(ParsedSequenceParameterSet, CurrentSequenceParameterSet))
+				{
+					MaxWidth = FMath::Max(MaxWidth, ParsedSequenceParameterSet.GetWidth());
+					MaxHeight = FMath::Max(MaxHeight, ParsedSequenceParameterSet.GetHeight());
+				}
+				else
+				{
+					UE_LOG(LogRtspMedia, Error, TEXT("Failed to parse sequence parameter set"));
+					return false;
+				}
+			}
+			CodecOptions.Emplace("max_width", FVariant(static_cast<int64>(MaxWidth)));
+			CodecOptions.Emplace("max_height", FVariant(static_cast<int64>(MaxHeight)));
+			UE_LOG(LogRtspMedia, Verbose, TEXT("Determined max image dimensions for decoder from SPS. MaxWidth: %d MaxHeight %d"), MaxWidth, MaxHeight);
+		}
+		else
+		{
+			UE_LOG(LogRtspMedia, Error, TEXT("Zero sequence parameter sets available within H.264 configuration"));
+			return false;
+		}
+	}
+	else
+	{
+		UE_LOG(LogRtspMedia, Error, TEXT("Failed to create AVC configuration record from CSD"));
+		return false;
+	}
+
+	IElectraCodecFactoryModule* FactoryModule = FModuleManager::Get().GetModulePtr<IElectraCodecFactoryModule>("ElectraCodecFactory");
+	if (!FactoryModule)
+	{
+		UE_LOG(LogRtspMedia, Error, TEXT("Failed to get ElectraCodecFactoryModule"));
+		return false;
+	}
+
+	// Codec format is: avc1.XXYYZZ
+	// Using 'avc1' here hard codes us to H.264 which is fine whilst this is the only protocol that we support.
+	// XX is profile_idc (Profile Indicator - Describes the encoder feature set e.g 'Baseline', 'Main' or 'High')
+	// YY is constraint flags (Indicates constraints the encoder followed, e.g. a Main profile encoder only used Baseline features)
+	// ZZ is level_idc (Level Indicator - Indicates the resources required for decoding. e.g. 31, Level 3.1 indicates 1280x720 @ 30 FPS)
+	//
+	// E.g. "avc1.42001F"
+	// 42 -> 66 -> Baseline profile
+	// 00 -> No constraint flags set
+	// 1F -> Level 31 -> Level 3.1 (1280 x 720 @ 30 FPS)
+	//
+	// To get this information the first four bytes of the SPS contain:
+	// 0: NAL Unit Header
+	// 1: profile_idc
+	// 2: constraint flags
+	// 3: level_idc
+
+	if (CodecFormat.IsEmpty())
+	{
+		UE_LOG(LogRtspMedia, Error, TEXT("CodecFormat was empty when preparing to create decoder"));
+		return false;
+	}
+	
+	TMap<FString, FVariant> FormatInfo;
+	CodecFactory = FactoryModule->GetBestFactoryForFormat(FormatInfo, CodecFormat, false, CodecOptions);
+	if (CodecFactory)
+	{
+		if (ElectraDecoder.IsValid())
+		{
+			ElectraDecoder->Close();
+		}
+		
+		ElectraDecoder = CodecFactory->CreateDecoderForFormat(CodecFormat, CodecOptions);
+
+		if (!ElectraDecoder.IsValid())
+		{
+			UE_LOG(LogRtspMedia, Error, TEXT("Failed to create decoder for format: %s"), *CodecFormat);
+			return false;
+		}
+	}
+	else
+	{
+		UE_LOG(LogRtspMedia, Error, TEXT("Failed to create codec factory for unsupported format: %s"), *CodecFormat);
+		return false;
+	}
+
+	BitstreamProcessor.Reset();
+	BitstreamProcessor = ElectraDecoder->CreateBitstreamProcessor();
+
+	if (!BitstreamProcessor.IsValid())
+	{
+		UE_LOG(LogRtspMedia, Error, TEXT("Decoder failed to create bitstream processor"));
+		return false;
+	}
+
+	return true;
+ }
+
+void FRtpDecoder::HandleOnH264NalUnit(TArray<uint8> InNalUnit, uint32 InTimestamp, bool bIsIDRFrame, bool bIsMarkerBitSet, bool bIsSlice)
+{
+	// Wait on an IDR frame first then allow all other units through.
+	if (bWaitForIDRFrame)
+	{
+		if (!bIsIDRFrame)
+		{
+			// Abort processing until we get one
+			return;
+		}
+		bWaitForIDRFrame = false;
+
+		// Ensure we're starting with a fresh frame buffer for the IDR frame
+		ResetFrameBuffer();
+	}
+
+	// Check for timestamp discontinuity in case NAL units have arrived out of order
+	// If we have data in our frame buffer from a previous NAL unit, check the timestamp
+	// of the new unit matches that of the frame as a whole.
+	if (FrameBuffer.Num() > 0 && InTimestamp != CurrentFrameTimestamp)
+	{
+		UE_LOG(LogRtspMedia, Warning, TEXT("Detected timestamp discontinuity, discarding frame"));
+		ResetFrameBuffer();
+		bWaitForIDRFrame = true;
+		return;
+	}
+
+	// Starting a new frame? Store the timestamp of this NAL unit as the timestamp for our frame
+	if (FrameBuffer.IsEmpty())
+	{
+		CurrentFrameTimestamp = InTimestamp;
+	}
+
+	// Append the NAL unit to our frame buffer
+	FrameBuffer.Append(InNalUnit);
+	
+	// Marker may indicate a complete access unit.
+	// However, we can't use it to reliably start decoding.
+	// https://www.rfc-editor.org/rfc/rfc6184#section-5.1
+	// See the bIsSlice check below. 
+	if (!bIsMarkerBitSet)
+	{
+		return;
+	}
+
+	// Some senders will set the marker bit on NAL units that contain only SPS/PPS data (metadata, no image data).
+	// The decoder requires a complete access unit with an image slice or decoding will fail.
+	// So once the marker bit is set we need to check that we have received a NAL unit that represents an image slice.
+	// Any metadata units we received with the same timestamp will be aggregated into this access unit.
+	// With the marker bit set and bIsSlice == true, we should have a complete unit containing image data ready for decoding.
+	if (!bIsSlice)
+	{
+		return;
+	}
+
+	bool bDecodedSuccessfully = false;
+
+	// After this point we should have a complete frame in the buffer
+	// Successful decoding or failure should reset the frame buffer from here on.
+	ON_SCOPE_EXIT
+	{
+		ResetFrameBuffer();
+
+		// We always set bWaitForIDRFrame to true on return EXCEPT for when we decoded successfully.
+		// In the success case we want to continue processing non IDR-frames
+		if (!bDecodedSuccessfully)
+		{
+			bWaitForIDRFrame = true;
+		}
+	};
+	
+	// An access unit represents one frame of data, consisting of multiple NAL units.
+	IElectraDecoder::FInputAccessUnit AccessUnit;
+	AccessUnit.Data = FrameBuffer.GetData();
+	AccessUnit.DataSize = FrameBuffer.Num();
+	AccessUnit.PTS = FTimespan::FromSeconds(static_cast<double>(InTimestamp) / VideoClockRate);
+	if (bIsIDRFrame)
+	{
+		AccessUnit.Flags = EElectraDecoderFlags::IsSyncSample;
+	}
+
+	if (!BitstreamProcessor.IsValid())
+	{
+		UE_LOG(LogRtspMedia, Warning, TEXT("BitstreamProcessor unavailable prior to decoding"));
+		return;
+	}
+
+	// Before decoding can take place, any platform specific preparation needs to be performed.
+	// On some platforms the big endian 4 byte length indicator we output in the depacketizer will be swapped
+	// for a H264 Annex B start code
+	// If we were using ElectraPlayer directly then we wouldn't need to perform this step.
+	// However, as we're directly integrating with the underlying decoders we need to prepare the data for processing.
+	TSharedPtr<IElectraDecoderBitstreamInfo, ESPMode::ThreadSafe> BitstreamInfo;
+	const IElectraDecoderBitstreamProcessor::EProcessResult ProcessResult = BitstreamProcessor->ProcessInputForDecoding(BitstreamInfo, AccessUnit, CodecOptions);
+	if (ProcessResult == IElectraDecoderBitstreamProcessor::EProcessResult::Error)
+	{
+		UE_LOG(LogRtspMedia, Warning, TEXT("BitstreamProcessor error: %s"), *BitstreamProcessor->GetLastError());
+		return;
+	}
+
+	if (!ElectraDecoder.IsValid())
+	{
+		UE_LOG(LogRtspMedia, Warning, TEXT("ElectraDecoder unavailable prior to decoding"));
+		return;
+	}
+
+	const IElectraDecoder::EDecoderError DecodeResult = ElectraDecoder->DecodeAccessUnit(AccessUnit, CodecOptions);
+
+	switch (DecodeResult)
+	{
+		case IElectraDecoder::EDecoderError::None:
+			bDecodedSuccessfully = true;
+			break;
+		case IElectraDecoder::EDecoderError::NoBuffer:
+			UE_LOG(LogRtspMedia, Verbose, TEXT("ElectraDecoder reports no buffer available, dropping unit"));
+			return;
+		case IElectraDecoder::EDecoderError::LostDecoder:
+			UE_LOG(LogRtspMedia, Warning, TEXT("Lost decoder reported, recreating"));
+			if (!CreateElectraDecoder(SequenceParameterSet, PictureParameterSet))
+			{
+				UE_LOG(LogRtspMedia, Error, TEXT("Failed to re-create ElectraDecoder after receiving a LostDecoder error"));
+				if (OnDecoderError.IsBound())
+				{
+					OnDecoderError.Execute();
+				}
+				bStopping.store(true);
+			}
+			return;
+		default:
+			IElectraDecoder::FError DecodeError = ElectraDecoder->GetError();
+			UE_LOG(LogRtspMedia, Warning, TEXT("ElectraDecoder error: \"%s\" Code: %d SDK Error: %d"), *DecodeError.GetMessage(), DecodeError.GetCode(), DecodeError.GetSdkCode());
+			break;
+	}
+
+	// Most frames will be picked up in the Run loop, but this will catch any frames that are decoded quickly.
+	PollDecodedOutput();
+}
+
+void FRtpDecoder::PollDecodedOutput()
+{
+	if (!ElectraDecoder.IsValid())
+	{
+		UE_LOG(LogRtspMedia, Warning, TEXT("Attempted to process decoded output without an electra decoder instance"));
+		return;
+	}
+
+	if (!TextureSamplePool.IsValid())
+	{
+		UE_LOG(LogRtspMedia, Error, TEXT("Texture sample buffer not available when processing decoded output"));
+		return;
+	}
+
+	while (ElectraDecoder->HaveOutput() == IElectraDecoder::EOutputStatus::Available)
+	{
+		TSharedPtr<IElectraDecoderOutput> Output = ElectraDecoder->GetOutput();
+		if (!Output.IsValid())
+		{
+			UE_LOG(LogRtspMedia, Warning, TEXT("Received null output from electra decoder"));
+			continue;
+		}
+
+		TSharedPtr<IElectraDecoderVideoOutput> VideoOutput = StaticCastSharedPtr<IElectraDecoderVideoOutput>(Output);
+		
+		if (!VideoOutput.IsValid())
+		{
+			UE_LOG(LogRtspMedia, Warning, TEXT("Cast from IElectraDecoderOutput to IElectraDecoderVideoOutput failed"));
+			continue;
+		}
+
+		if (VideoOutput->GetOutputType() == IElectraDecoderVideoOutput::EOutputType::DoNotOutput)
+		{
+			// There is a problem with the output, discard it.
+			UE_LOG(LogRtspMedia, Warning, TEXT("Video output type is 'DoNotOutput'. Discarding"));
+			continue;
+		}
+		
+		TSharedPtr<FElectraTextureSample, ESPMode::ThreadSafe> TextureSample = TextureSamplePool->AcquireShared();
+		if (!TextureSample.IsValid())
+		{
+			UE_LOG(LogRtspMedia, Error, TEXT("Failed to acquire texture sample from pool"));
+			continue;
+		}
+
+		// Get video dimensions for:
+		// - Providing format information to the Media framework
+		// - Populating the platform video resource options below
+		const int32 VideoWidth = VideoOutput->GetWidth();
+		const int32 VideoHeight = VideoOutput->GetHeight();
+		
+		// If the video dimensions change reset any platform video resources we may be using,
+		// updated the cached dimensions and inform the delegate on the game thread.
+		if (VideoWidth != CurrentVideoWidth || VideoHeight != CurrentVideoHeight)
+		{
+			ResetPlatformVideoResource();
+			CurrentVideoWidth = VideoWidth;
+			CurrentVideoHeight = VideoHeight;
+			if (OnVideoDimensions.IsBound())
+			{
+				OnVideoDimensions.Execute(VideoWidth, VideoHeight);
+			}
+		}
+
+		// For some platforms we must specify the required resources for the texture
+		// On Apple platforms this information is ignored, but it is used on Windows and Linux
+		if (!PlatformResource)
+		{
+			TMap<FString, FVariant> PlatformOptions;
+			PlatformOptions.Emplace(TEXT("max_width"), FVariant(static_cast<int64>(VideoWidth)));
+			PlatformOptions.Emplace(TEXT("max_height"), FVariant(static_cast<int64>(VideoHeight)));
+			PlatformResource = FElectraDecodersPlatformResources::CreatePlatformVideoResource(this, PlatformOptions);
+		}
+
+		// Calculate frame duration from presentation timestamps
+		const FTimespan CurrentPresentationTimestamp = VideoOutput->GetPTS();
+		if (!LastPresentationTimestamp.IsZero() && CurrentPresentationTimestamp > LastPresentationTimestamp)
+		{
+			const FTimespan FrameDuration = CurrentPresentationTimestamp - LastPresentationTimestamp;
+			FrameDurationTicks = FrameDuration.GetTicks();
+
+			NumFramesSinceLastUpdate++;
+			
+			// Update frame rate
+			constexpr double FrameRateUpdateIntervalSeconds = 1.0;
+			const double Now = FPlatformTime::Seconds();
+			const double TimeSinceLastUpdate = Now - LastFrameRateUpdate;
+			if (TimeSinceLastUpdate > FrameRateUpdateIntervalSeconds)
+			{
+				const double Average = TimeSinceLastUpdate / NumFramesSinceLastUpdate;
+				const double NewFrameRate = 1.0 / Average;
+				FrameRate.store(FMath::RoundToDouble(NewFrameRate), std::memory_order_relaxed);
+				NumFramesSinceLastUpdate = 0;
+				LastFrameRateUpdate = Now;
+			}
+		}
+		LastPresentationTimestamp = CurrentPresentationTimestamp;
+
+		// Populate buffer timing properties 
+		TSharedPtr<Electra::FParamDict, ESPMode::ThreadSafe> BufferProperties = MakeShared<Electra::FParamDict, ESPMode::ThreadSafe>();
+
+		// Electra uses its own time values that we need to map
+		// One FTimespan tick is 100 nanoseconds
+		const int64 PresentationTimeTicks = CurrentPresentationTimestamp.GetTicks();
+		BufferProperties->Set(IDecoderOutputOptionNames::PTS, Electra::FVariantValue(Electra::FTimeValue(PresentationTimeTicks)));
+		BufferProperties->Set(IDecoderOutputOptionNames::Duration, Electra::FVariantValue(Electra::FTimeValue(FrameDurationTicks)));
+
+		FString SetupOutputTextureSampleError;
+		const bool bSuccess = FElectraDecodersPlatformResources::SetupOutputTextureSample(
+			SetupOutputTextureSampleError, 
+			TextureSample, 
+			VideoOutput, 
+			BufferProperties, 
+			PlatformResource
+		);
+
+		if (!bSuccess)
+		{
+			UE_LOG(LogRtspMedia, Warning, TEXT("Failed to set up texture sample: %s"), *SetupOutputTextureSampleError);
+			continue;
+		}
+
+		TextureSample->SetTime(FMediaTimeStamp(CurrentPresentationTimestamp));
+
+		Samples->AddVideo(TextureSample.ToSharedRef());
+	}
+}
+
+void FRtpDecoder::ResetFrameBuffer()
+{
+	FrameBuffer.Empty();
+	CurrentFrameTimestamp = 0;
+}
+
+void FRtpDecoder::ResetPlatformVideoResource()
+{
+	if (PlatformResource)
+	{
+		FElectraDecodersPlatformResources::ReleasePlatformVideoResource(this, PlatformResource);
+		PlatformResource = nullptr;
+	}
+}
diff --git a/Engine/Plugins/Media/RTSPMedia/Source/RTSPMedia/Private/RTP/RtpDecoder.h b/Engine/Plugins/Media/RTSPMedia/Source/RTSPMedia/Private/RTP/RtpDecoder.h
new file mode 100644
--- /dev/null
+++ b/Engine/Plugins/Media/RTSPMedia/Source/RTSPMedia/Private/RTP/RtpDecoder.h
@@ -0,0 +1,122 @@
+// Copyright Epic Games, Inc. All Rights Reserved.
+
+#pragma once
+
+#include "RtspMediaDefaults.h"
+
+#include "RtpPacket.h"
+
+#include "Containers/Queue.h"
+#include "CoreMinimal.h"
+#include "ElectraDecodersPlatformResources.h"
+#include "HAL/Runnable.h"
+#include "HAL/RunnableThread.h"
+#include "MediaSamples.h"
+#include "Misc/Variant.h"
+#include "Templates/SharedPointer.h"
+
+class IElectraDecoder;
+class IElectraDecoderBitstreamProcessor;
+class IElectraCodecFactory;
+class IElectraDecoderVideoOutput;
+class FElectraTextureSample;
+class FElectraTextureSamplePool;
+class FMediaSamples;
+class FRtpH264Depacketizer;
+
+struct FSdpSession;
+
+// Provides Width and Height, executed on the decoder thread.
+DECLARE_DELEGATE_TwoParams(FOnVideoDimensions, int32, int32);
+
+DECLARE_DELEGATE(FOnDecoderError);
+
+/**
+ * Uses a dedicated worker thread for:
+ * - Depacketizing H.264 NAL units
+ * - H.264 Decoder execution
+ *
+ * The Initialize, Start and Shutdown methods are all expected to be called from the game thread.
+ * EnqueuePacket is expected to be called from the RtspClient worker thread to avoid
+ * a hop to the game thread in order to minimise latency.
+ */
+class FRtpDecoder : public FRunnable
+{
+public:
+	FRtpDecoder();
+	virtual ~FRtpDecoder();
+	
+	bool Initialize(
+		FMediaSamples* InSamples,
+		const FSdpSession& InSession,
+		int64 InMaxFragmentBufferSize,
+		int64 InDecoderBufferSize,
+		uint32 InDecoderPollIntervalMs);
+	
+	bool Start();
+	void Shutdown();
+	void EnqueuePacket(FRtpPacket InPacket);
+
+	double GetFrameRate() const;
+
+	FOnVideoDimensions OnVideoDimensions;
+	FOnDecoderError OnDecoderError;
+
+protected:
+	// ~FRunnable Start
+	virtual bool Init() override;
+	virtual uint32 Run() override;
+	virtual void Stop() override;
+	// ~FRunnable Stop
+
+private:
+
+	bool CreateElectraDecoder(const TArray<uint8>& InSequenceParameterSet, const TArray<uint8>& InPictureParameterSet);
+	void HandleOnH264NalUnit(TArray<uint8> InNalUnit, uint32 InTimestamp, bool bIsIDRFrame, bool bIsMarkerBitSet, bool bIsSlice);
+	void PollDecodedOutput();
+	
+	void ResetFrameBuffer();
+	void ResetPlatformVideoResource();
+
+	TQueue<FRtpPacket, EQueueMode::Spsc> Queue;
+
+	// Thread Management
+	std::atomic<bool> bStopping;
+	TUniquePtr<FRunnableThread> Thread;
+	FEvent* WaitEvent;
+
+	// Decoding Tools
+	TSharedPtr<IElectraDecoderBitstreamProcessor, ESPMode::ThreadSafe> BitstreamProcessor;
+	TSharedPtr<IElectraCodecFactory> CodecFactory;
+	TSharedPtr<IElectraDecoder> ElectraDecoder;
+	TUniquePtr<FRtpH264Depacketizer> H264Depacketizer;
+
+	// Decoding Configuration
+	TArray<uint8> SequenceParameterSet;
+	TArray<uint8> PictureParameterSet;
+	TMap<FString, FVariant> CodecOptions;
+	int64 MaxFragmentBufferSize = RtspMedia::Default::MaxFragmentBufferSizeBytes;
+	int64 DecoderBufferSize = RtspMedia::Default::DecoderBufferSize;
+	uint32 DecoderPollIntervalMs = RtspMedia::Default::DecoderPollIntervalMs;
+
+	// Decoding State
+	uint32 CurrentFrameTimestamp = 0;
+	TArray<uint8> FrameBuffer;
+	TUniquePtr<FSdpSession> SdpSession;
+	uint32 VideoClockRate = 90000;
+	bool bWaitForIDRFrame = true;
+
+	// Texture Sample Generation
+	int32 CurrentVideoWidth = 0;
+	int32 CurrentVideoHeight = 0;
+	FTimespan LastPresentationTimestamp;
+	FElectraDecodersPlatformResources::IDecoderPlatformResource* PlatformResource = nullptr;
+	TUniquePtr<FElectraTextureSamplePool> TextureSamplePool;
+	int64 FrameDurationTicks = 333333;
+	FMediaSamples* Samples = nullptr;
+	
+	// Frame rate will be updated on the decoder thread and read by the game thread for UI display
+	std::atomic<double> FrameRate;
+	double LastFrameRateUpdate = 0.0;
+	int32 NumFramesSinceLastUpdate = 0;
+};
diff --git a/Engine/Plugins/Media/RTSPMedia/Source/RTSPMedia/Private/RTP/RtpH264Depacketizer.cpp b/Engine/Plugins/Media/RTSPMedia/Source/RTSPMedia/Private/RTP/RtpH264Depacketizer.cpp
new file mode 100644
--- /dev/null
+++ b/Engine/Plugins/Media/RTSPMedia/Source/RTSPMedia/Private/RTP/RtpH264Depacketizer.cpp
@@ -0,0 +1,261 @@
+// Copyright Epic Games, Inc. All Rights Reserved.
+
+#include "RtpH264Depacketizer.h"
+
+#include "RtspMediaConstants.h"
+
+void FRtpH264Depacketizer::ProcessRtpPayload(TArray<uint8> InRtpPayload, uint16 InSequenceNumber, uint32 InTimestamp, bool bInMarker)
+{
+	if (InRtpPayload.IsEmpty())
+	{
+		UE_LOG(LogRtspMedia, Warning, TEXT("H264 depacketizer was provided with an empty payload"));
+		ResetFragmentBuffer();
+		return;
+	}
+
+	if (!bIsFirstPacket)
+	{
+		const uint16 ExpectedSequenceNumber = LastSequenceNumber + 1;
+		if (InSequenceNumber != ExpectedSequenceNumber)
+		{
+			UE_LOG(LogRtspMedia, Warning, TEXT("H264 depacketizer expected sequence number %d but got %d"), ExpectedSequenceNumber, InSequenceNumber);
+			// Drop any partial fragments we might have at this point.
+			ResetFragmentBuffer();
+		}
+	}
+	else 
+	{
+		bIsFirstPacket = false;
+	}
+
+	LastSequenceNumber = InSequenceNumber;
+
+	// First byte contains the NAL header:
+	// F - Forbidden zero bit (always 0)
+	// NRI - 2 bits - Importance indicator (0-3)
+	// Type - 5 bits - NAL packetization type
+	// Clear away the F and NRI bits with 0x1F (00011111) to get the packetization type.
+	const uint8 PacketizationType = InRtpPayload[0] & 0x1F;
+
+	// 1 - 23 are complete NAL units
+	if (PacketizationType > 0 && PacketizationType < 24)
+	{
+		HandleSingleUnit(InRtpPayload, InTimestamp, bInMarker);
+	}
+	// 24 is a RTP specific type indicating STAP-A (Aggregate unit)
+	else if (PacketizationType == 24)
+	{
+		HandleAggregate(InRtpPayload, InTimestamp, bInMarker);
+	}
+	// 28 is another RTP specific type indicating FU-A (Fragment unit)
+	else if (PacketizationType == 28)
+	{
+		HandleFragment(InRtpPayload, InTimestamp, bInMarker);
+	}
+	else
+	{
+		UE_LOG(LogRtspMedia, Warning, TEXT("H264 depacketizer encountered unknown NAL unit type: %d"), PacketizationType);
+
+		// Only reset if we're expecting another fragment.
+		// It could be a NalType that we don't support. e.g. STAP-B, MTAP, FU-B.
+		if (bFragmentInProgress)
+		{
+			ResetFragmentBuffer();
+		}
+	}
+}
+
+void FRtpH264Depacketizer::SetMaxFragmentBufferSize(int64 InMaxFragmentBufferSize)
+{
+	if (ensureMsgf(InMaxFragmentBufferSize > 0, TEXT("Attempted to set a zero or negative max fragment buffer size")))
+	{
+		MaxFragmentBufferSize = InMaxFragmentBufferSize;
+	}
+}
+
+void FRtpH264Depacketizer::ResetFragmentBuffer()
+{
+	FragmentBuffer.Empty();
+	bFragmentInProgress = false;
+	FragmentTimestamp = 0;
+}
+
+// Some senders (E.g. The RootEncoder library on Android) incorrectly include Annex B start codes at the
+// start of certain NAL units.
+// This could be the four byte start code 0x00 0x00 0x00 0x01
+// Or the three byte start code 0x00 0x00 0x01
+TArrayView<const uint8> FRtpH264Depacketizer::StripAnnexBStartCodes(TArrayView<const uint8> InNalData)
+{
+	const uint8* Data = InNalData.GetData();
+	const int32 Size = InNalData.Num();
+
+	// 4 byte start code
+	if (Size >= 4 && Data[0] == 0x00 && Data[1] == 0x00 && Data[2] == 0x00 && Data[3] == 0x01)
+	{
+		UE_LOG(LogRtspMedia, VeryVerbose, TEXT("Stripping 4 byte Annex B start code from non-compliant NAL unit"));
+		return TArrayView<const uint8>(Data + 4, Size - 4);
+	}
+
+	// 3 byte start code
+	if (Size >= 3 && Data[0] == 0x00 && Data[1] == 0x00 && Data[2] == 0x01)
+	{
+		UE_LOG(LogRtspMedia, VeryVerbose, TEXT("Stripping 3 byte Annex B start code from non-compliant NAL unit"));
+		return TArrayView<const uint8>(Data + 3, Size - 3);
+	}
+
+	return InNalData;
+}
+
+void FRtpH264Depacketizer::HandleSingleUnit(TArrayView<const uint8> InPayload, uint32 InTimestamp, bool bInMarker)
+{
+	EmitNalUnit(InPayload, InTimestamp, bInMarker);
+}
+
+void FRtpH264Depacketizer::HandleAggregate(TArrayView<const uint8> InPayload, uint32 InTimestamp, bool bInMarker)
+{
+	const uint8* Data = InPayload.GetData();
+	const int32 DataSize = InPayload.Num();
+
+	// Skip over the NAL type byte we already used to determine that this is an aggregate packet.
+	int32 Offset = 1;
+
+	while (Offset + 2 <= DataSize)
+	{
+		// Multi-byte values are in network byte order.
+		const uint16 NalSize = (static_cast<uint16>(Data[Offset]) << 8) | Data[Offset + 1];
+		Offset += 2;
+
+		if (NalSize == 0)
+		{
+			UE_LOG(LogRtspMedia, Warning, TEXT("Detected NAL size of zero within aggregate packet"));
+			continue;
+		}
+
+		if (Offset + NalSize > DataSize)
+		{
+			UE_LOG(LogRtspMedia, Warning, TEXT("H264 depacketizer encountered malformed aggregate payload"));
+			break;
+		}
+
+		const TArrayView<const uint8> NalData(Data + Offset, NalSize);
+
+		// Only the last NAL in the aggregate can potentially be marked as the last unit of the frame.
+		bool bIsLastNal = false;
+		if (Offset + NalSize >= DataSize)
+		{
+			bIsLastNal = true;
+		}
+		EmitNalUnit(NalData, InTimestamp, bIsLastNal && bInMarker);
+		Offset += NalSize;
+	}
+}
+
+void FRtpH264Depacketizer::HandleFragment(TArrayView<const uint8> InPayload, uint32 InTimestamp, bool bInMarker)
+{
+	const uint8* Data = InPayload.GetData();
+	const int32 DataSize = InPayload.Num();
+
+	if (DataSize < 2)
+	{
+		UE_LOG(LogRtspMedia, Warning, TEXT("H264 depacketizer encountered fragment payload containing less than 2 bytes"));
+		return;
+	}
+
+	const uint8 FragmentIndicator = Data[0];
+	const uint8 FragmentHeader = Data[1];
+
+	// Fragment header contains
+	// Start bit (answers the question: is this the first fragment of a NAL unit?)
+	// End bit
+	// Reserved bit
+	// Original NAL unit type - 5 bits (1 - 23)
+	const bool bIsStart = (FragmentHeader >> 7) & 0x01;
+	const bool bIsEnd = (FragmentHeader >> 6) & 0x01;
+
+	constexpr uint32 FragmentOffset = 2;
+	const TArrayView<const uint8> FragmentData(Data + FragmentOffset, DataSize - FragmentOffset);
+
+	if (bIsStart)
+	{
+		ResetFragmentBuffer();
+		// When packetized as FU-A the original NAL header byte is not sent.
+		// Instead, the information contained within is split between the first fragment
+		// indicator byte and the fragment header in this fragment payload.
+		// We take the first 3 bits from the indicator and the last 5 bits from the header.
+		// This gives us the original NAL unit header required at the start of a NAL unit for decoding.
+		const uint8 NalType = FragmentHeader & 0x1F;
+		const uint8 ReconstructedNalHeader = (FragmentIndicator & 0xE0) | NalType;
+		FragmentBuffer.Add(ReconstructedNalHeader);
+		bFragmentInProgress = true;
+
+		// We use the timestamp of the first fragment
+		FragmentTimestamp = InTimestamp;
+	}
+
+	if (bFragmentInProgress)
+	{
+		// Check if we're about to exceed our max fragment buffer size
+		if (static_cast<int64>(FragmentBuffer.Num()) + static_cast<int64>(FragmentData.Num()) > MaxFragmentBufferSize)
+		{
+			UE_LOG(LogRtspMedia, Warning, TEXT("Max fragment buffer size (%lld) exceeded. Dropping fragmented NAL unit."), MaxFragmentBufferSize);
+			ResetFragmentBuffer();
+			return;
+		}
+		FragmentBuffer.Append(FragmentData.GetData(), FragmentData.Num());
+	}
+
+	if (bIsEnd && bFragmentInProgress)
+	{
+		EmitNalUnit(FragmentBuffer, FragmentTimestamp, bInMarker);
+		ResetFragmentBuffer();
+	}
+}
+
+void FRtpH264Depacketizer::EmitNalUnit(TArrayView<const uint8> InNalUnit, uint32 InTimestamp, bool bInMarker)
+{
+	if (!ensureMsgf(!InNalUnit.IsEmpty(), TEXT("NAL unit is empty at point of emission from depacketizer")))
+	{
+		return;
+	}
+
+	// Remove any Annex B start codes that a non-compliant sender may have included in the NAL unit.
+	// These shouldn't be present in RTP packetized NAL units. 
+	// https://www.rfc-editor.org/rfc/rfc6184.txt#:~:text=Scope%0A%0A%20%20%20This%20payload-,specification,-can%20only%20be
+	InNalUnit = StripAnnexBStartCodes(InNalUnit);
+	
+	if (InNalUnit.IsEmpty())
+	{
+		UE_LOG(LogRtspMedia, VeryVerbose, TEXT("Non-compliant NAL unit is empty after stripping annex B start codes"));
+		return;
+	}
+
+	// Determine whether this NAL unit represents an IDR-Frame
+	// This is a special kind of I-Frame that allows the decoder to clear
+	// its internal state. The last 5 bits of the first byte in the unit
+	// contain the NAL type. 5 == an IDR frame.
+	// 0x1F = 00011111
+	const uint8 NalType = InNalUnit[0] & 0x1F;
+	const bool bIsIdrFrame = NalType == 5;
+
+	// Because certain senders will set the marker bit to 1 even if we haven't
+	// received image data for that timestamp we need to also track if the NAL type is an image 'slice'.
+	// I.e. Not a NAL unit containing only metadata.
+	const bool bIsSlice = NalType >= 1 && NalType <= 5;
+
+	// The final NAL unit output needs to be in the AVCC format with the NAL unit data itself
+	// prefixed with a 4-byte big endian value describing the length of the contained NAL unit.
+
+	const int32 NalUnitSize = InNalUnit.Num();
+	const uint8 HeaderByte0 = (NalUnitSize >> 24) & 0xFF;
+	const uint8 HeaderByte1 = (NalUnitSize >> 16) & 0xFF;
+	const uint8 HeaderByte2 = (NalUnitSize >> 8) & 0xFF;
+	const uint8 HeaderByte3 = NalUnitSize & 0xFF;
+
+	TArray<uint8> AvccNalUnit;
+	AvccNalUnit.Add(HeaderByte0);
+	AvccNalUnit.Add(HeaderByte1);
+	AvccNalUnit.Add(HeaderByte2);
+	AvccNalUnit.Add(HeaderByte3);
+	AvccNalUnit.Append(InNalUnit.GetData(), NalUnitSize);
+	OnH264NalUnit.ExecuteIfBound(MoveTemp(AvccNalUnit), InTimestamp, bIsIdrFrame, bInMarker, bIsSlice);
+}
diff --git a/Engine/Plugins/Media/RTSPMedia/Source/RTSPMedia/Private/RTP/RtpH264Depacketizer.h b/Engine/Plugins/Media/RTSPMedia/Source/RTSPMedia/Private/RTP/RtpH264Depacketizer.h
new file mode 100644
--- /dev/null
+++ b/Engine/Plugins/Media/RTSPMedia/Source/RTSPMedia/Private/RTP/RtpH264Depacketizer.h
@@ -0,0 +1,39 @@
+// Copyright Epic Games, Inc. All Rights Reserved.
+
+#pragma once
+
+#include "RtspMediaDefaults.h"
+
+#include "CoreMinimal.h"
+
+// NAL Unit, Timestamp, bIsIdrFrame, bIsMarkerBitSet, bIsSlice 
+DECLARE_DELEGATE_FiveParams(FOnH264NalUnit, TArray<uint8>, uint32, bool, bool, bool);
+
+class FRtpH264Depacketizer
+{
+public:
+	void ProcessRtpPayload(TArray<uint8> InRtpPayload, uint16 InSequenceNumber, uint32 InTimestamp, bool bInMarker);
+	void SetMaxFragmentBufferSize(int64 InMaxFragmentBufferSize);
+
+	FOnH264NalUnit OnH264NalUnit;
+private:
+	void HandleSingleUnit(TArrayView<const uint8> InPayload, uint32 InTimestamp, bool bInMarker); // Complete NAL
+	void HandleAggregate(TArrayView<const uint8> InPayload, uint32 InTimestamp, bool bInMarker); // STAP-A containing multiple NALs
+	void HandleFragment(TArrayView<const uint8> InPayload, uint32 InTimestamp, bool bInMarker); // FU-A NAL fragment
+
+	void EmitNalUnit(TArrayView<const uint8> InNalUnit, uint32 InTimestamp, bool bInMarker);
+
+	void ResetFragmentBuffer();
+
+	TArrayView<const uint8> StripAnnexBStartCodes(TArrayView<const uint8> InNalData);
+
+	// State for fragment reassembly
+	TArray<uint8> FragmentBuffer;
+	int64 MaxFragmentBufferSize = RtspMedia::Default::MaxFragmentBufferSizeBytes;
+	bool bFragmentInProgress = false;
+	uint32 FragmentTimestamp = 0;
+
+	// Sequence tracking
+	uint16 LastSequenceNumber = 0;
+	bool bIsFirstPacket = true;
+};
\ No newline at end of file
diff --git a/Engine/Plugins/Media/RTSPMedia/Source/RTSPMedia/Private/RTP/RtpJitterBuffer.cpp b/Engine/Plugins/Media/RTSPMedia/Source/RTSPMedia/Private/RTP/RtpJitterBuffer.cpp
new file mode 100644
--- /dev/null
+++ b/Engine/Plugins/Media/RTSPMedia/Source/RTSPMedia/Private/RTP/RtpJitterBuffer.cpp
@@ -0,0 +1,202 @@
+// Copyright Epic Games, Inc. All Rights Reserved.
+
+#include "RtpJitterBuffer.h"
+
+#include "RtspMediaConstants.h"
+
+#include "HAL/PlatformTime.h"
+
+bool FRtpJitterBuffer::Initialize(const uint32 InBufferDepthMs, const uint32 InClockRate)
+{
+	if (InClockRate == 0)
+	{
+		UE_LOG(LogRtspMedia, Error, TEXT("Attempted to initialize RTP jitter buffer with a zero clock rate"));
+		return false;
+	}
+
+	ClockRate = InClockRate;
+	PacketBufferDepthMs = InBufferDepthMs;
+	PacketBufferDepthSeconds = PacketBufferDepthMs / 1000.0;
+
+	ResetBuffer();
+	ResetPacing();
+
+	bInitialized = true;
+	
+	UE_LOG(LogRtspMedia, Verbose, TEXT("Initialized RTP jitter buffer with buffer depth (ms): %u Clock rate: %u"), PacketBufferDepthMs, ClockRate);
+	return true;
+}
+
+void FRtpJitterBuffer::EnqueuePacket(FRtpPacket InPacket)
+{
+	if (!ensureMsgf(bInitialized, TEXT("RTP jitter must be initialized before enqueuing packets")))
+	{
+		return;
+	}
+
+	// If we're in the filling phase and we haven't enqueued a packet yet
+	if (bFilling && !bFirstPacketEnqueued)
+	{
+		FillStartTimeSeconds = FPlatformTime::Seconds();
+	}
+
+	// First packet or a packet with a distinct timestamp
+	const uint32 Timestamp = InPacket.Header.Timestamp;
+	if (!bFirstPacketEnqueued || Timestamp != LastEnqueuedTimestamp)
+	{
+		++QueuedFrameCount;
+	}
+
+	UE_LOG(
+		LogRtspMedia,
+		VeryVerbose,
+		TEXT("RTP jitter buffer enqueuing packet with sequence number: %u RTP timestamp: %u"),
+		InPacket.Header.SequenceNumber,
+		InPacket.Header.Timestamp);
+	
+	PacketBuffer.Emplace(MoveTemp(InPacket));
+
+	LastEnqueuedTimestamp = Timestamp;
+	bFirstPacketEnqueued = true;
+}
+
+bool FRtpJitterBuffer::DequeueIfReady(FRtpPacket& OutPacket)
+{
+	if (!ensureMsgf(bInitialized, TEXT("RTP jitter must be initialized before dequeuing packets")))
+	{
+		return false;
+	}
+	
+	if (PacketBuffer.IsEmpty())
+	{
+		return false;
+	}
+
+	if (!ensureMsgf(QueueIndex < PacketBuffer.Num(), TEXT("QueueIndex: %d exceeds the number of packets in the jitter buffer: %d Discarding pending packets"), QueueIndex, PacketBuffer.Num()))
+	{
+		ResetBuffer();
+		ResetPacing();
+		return false;
+	}
+	
+	const double Now = FPlatformTime::Seconds();
+
+	FRtpPacket& FirstPacket = PacketBuffer[QueueIndex];
+
+	// Skip buffer filling and pacing if the configured packet buffer depth is 0
+	if (PacketBufferDepthMs > 0)
+	{
+		if (bFilling)
+		{
+			if (!bFirstPacketEnqueued)
+			{
+				return false;
+			}
+		
+			if (Now - FillStartTimeSeconds < PacketBufferDepthSeconds)
+			{
+				return false;
+			}
+
+			// The buffer is full, allow draining to take place.
+			NextReleaseTimeSeconds = Now;
+			bFilling = false;
+		}
+
+		// If the packet contains the same timestamp as the last dequeued packet, release it immediately.
+		const bool bSameTimestamp = bFirstPacketDequeued && FirstPacket.Header.Timestamp == LastDequeuedTimestamp;
+		if (!bSameTimestamp)
+		{
+			// If it's not release time yet, abort.
+			if (EstimatedFrameIntervalSeconds > 0.0 && Now < NextReleaseTimeSeconds)
+			{
+				return false;
+			}
+
+			// Calculate the initial estimated frame interval and target frame count if not set 
+			if (EstimatedFrameIntervalSeconds == 0.0 && QueuedFrameCount >= 2)
+			{
+				// The cast to int32 is intentional for wraparound safety.
+				const int32 Span = static_cast<int32>(PacketBuffer.Last().Header.Timestamp - FirstPacket.Header.Timestamp);
+				EstimatedFrameIntervalSeconds = (static_cast<double>(Span) / ClockRate) / (QueuedFrameCount - 1);
+			}
+
+			// Adjust the estimated frame interval for any deviation from the previously calculated frame interval
+			if (EstimatedFrameIntervalSeconds > 0.0)
+			{
+				const double TargetFrameCount = PacketBufferDepthSeconds / EstimatedFrameIntervalSeconds;
+				
+				// Normalized deviation of actual queued frame count vs the target frame count
+				const double Error = (static_cast<double>(QueuedFrameCount) - TargetFrameCount) / TargetFrameCount;
+
+				// Prevent the adjustment from becoming negative when an extreme burst of packets is present in the buffer.
+				const double Adjustment = FMath::Max(1.0 - BufferLevelGain * Error, BufferLevelGain);
+				
+				// Adjust the estimated frame interval by the error to bring it closer to the actual frame interval
+				EstimatedFrameIntervalSeconds *= Adjustment;
+			}
+			
+			// We will release now so compute the next release time
+			NextReleaseTimeSeconds += EstimatedFrameIntervalSeconds;
+		}
+	}
+
+	if (!bFirstPacketDequeued || FirstPacket.Header.Timestamp != LastDequeuedTimestamp)
+	{
+		--QueuedFrameCount;
+	}
+	
+	UE_LOG(
+		LogRtspMedia,
+		VeryVerbose, 
+		TEXT("Dequeuing RTP packet from jitter buffer. Sequence Number: %u Timestamp: %u"),
+		FirstPacket.Header.SequenceNumber,
+		FirstPacket.Header.Timestamp);
+		
+	bFirstPacketDequeued = true;
+	LastDequeuedTimestamp = FirstPacket.Header.Timestamp;
+	OutPacket = MoveTemp(FirstPacket);
+	++QueueIndex;
+	CompactPacketBuffer();
+
+	return true;
+}
+
+int32 FRtpJitterBuffer::NumQueuedPackets() const
+{
+	return PacketBuffer.Num() - QueueIndex;
+}
+
+void FRtpJitterBuffer::CompactPacketBuffer()
+{
+	if (QueueIndex > NumQueuedPackets())
+	{
+		PacketBuffer.RemoveAt(0, QueueIndex);
+		QueueIndex = 0;
+	}
+
+	if (NumQueuedPackets() == 0)
+	{
+		ResetBuffer();
+		ResetPacing();
+	}
+}
+
+void FRtpJitterBuffer::ResetBuffer()
+{
+	PacketBuffer.Empty();
+	QueueIndex = 0;
+}
+
+void FRtpJitterBuffer::ResetPacing()
+{
+	EstimatedFrameIntervalSeconds = 0.0;
+	bFilling = true;
+	FillStartTimeSeconds = 0.0;
+	bFirstPacketDequeued = false;
+	bFirstPacketEnqueued = false;
+	LastDequeuedTimestamp = 0;
+	LastEnqueuedTimestamp = 0;
+	NextReleaseTimeSeconds = 0.0;
+	QueuedFrameCount = 0;
+}
diff --git a/Engine/Plugins/Media/RTSPMedia/Source/RTSPMedia/Private/RTP/RtpJitterBuffer.h b/Engine/Plugins/Media/RTSPMedia/Source/RTSPMedia/Private/RTP/RtpJitterBuffer.h
new file mode 100644
--- /dev/null
+++ b/Engine/Plugins/Media/RTSPMedia/Source/RTSPMedia/Private/RTP/RtpJitterBuffer.h
@@ -0,0 +1,52 @@
+// Copyright Epic Games, Inc. All Rights Reserved.
+
+#pragma once
+
+#include "CoreMinimal.h"
+#include "RtpPacket.h"
+
+/**
+ * The jitter buffer assumes that packets are enqueued in order.
+ * This holds true naturally when the TCP interleaved connection mode is used.
+ * For UDP transport packets must be ordered before they are added to the buffer.
+ */
+class FRtpJitterBuffer
+{
+public:
+	bool Initialize(const uint32 InBufferDepthMs, const uint32 InClockRate);
+	void EnqueuePacket(FRtpPacket InPacket);
+	bool DequeueIfReady(FRtpPacket& OutPacket);
+
+private:
+	void CompactPacketBuffer();
+	int32 NumQueuedPackets() const;
+	void ResetBuffer();
+	void ResetPacing();
+
+	static constexpr double BufferLevelGain = 0.02;
+
+	// Configuration
+	uint32 ClockRate = 0;
+	bool bInitialized = false;
+	uint32 PacketBufferDepthMs = 0;
+	double PacketBufferDepthSeconds = 0.0;
+	
+	// Buffer
+	TArray<FRtpPacket> PacketBuffer;
+	int32 QueueIndex = 0;
+
+	// Fill Phase
+	bool bFilling = true;
+	double FillStartTimeSeconds = 0.0;
+
+	// Frame Pacing
+	double EstimatedFrameIntervalSeconds = 0.0;
+	double NextReleaseTimeSeconds = 0.0;
+	int32 QueuedFrameCount = 0;
+
+	// Enqueue/Dequeue Tracking
+	uint32 LastEnqueuedTimestamp = 0;
+	bool bFirstPacketEnqueued = false;
+	uint32 LastDequeuedTimestamp = 0;
+	bool bFirstPacketDequeued = false;
+};
diff --git a/Engine/Plugins/Media/RTSPMedia/Source/RTSPMedia/Private/RTP/RtpPacket.h b/Engine/Plugins/Media/RTSPMedia/Source/RTSPMedia/Private/RTP/RtpPacket.h
new file mode 100644
--- /dev/null
+++ b/Engine/Plugins/Media/RTSPMedia/Source/RTSPMedia/Private/RTP/RtpPacket.h
@@ -0,0 +1,44 @@
+// Copyright Epic Games, Inc. All Rights Reserved.
+
+#pragma once
+
+#include "CoreMinimal.h"
+
+struct FRtpHeader 
+{
+	// 2 bits - value should always be 2
+	uint8 Version = 0;
+
+	// P bit
+	bool bPadding = false;
+
+	// X bit
+	bool bExtension = false;
+
+	// 4 bits - Contributing source count
+	// Used when sources are mixed by a server
+	// Typically 0 in video streaming 
+	uint8 CsrcCount = 0;
+
+	// M bit - indicates last packet of frame
+	bool bMarker = false;
+
+	// 7 bits
+	uint8 PayloadType = 0;
+
+	// 2 bytes - for ordering
+	uint16 SequenceNumber = 0;
+
+	// 4 bytes - 90kHz clock
+	uint32 Timestamp = 0;
+
+	// 4 bytes - Synchronization source ID
+	// Random identifier
+	uint32 Ssrc = 0;
+};
+
+struct FRtpPacket
+{
+	FRtpHeader Header;
+	TArray<uint8> Payload;
+};
diff --git a/Engine/Plugins/Media/RTSPMedia/Source/RTSPMedia/Private/RTP/RtpParser.cpp b/Engine/Plugins/Media/RTSPMedia/Source/RTSPMedia/Private/RTP/RtpParser.cpp
new file mode 100644
--- /dev/null
+++ b/Engine/Plugins/Media/RTSPMedia/Source/RTSPMedia/Private/RTP/RtpParser.cpp
@@ -0,0 +1,128 @@
+// Copyright Epic Games, Inc. All Rights Reserved.
+
+#include "RtpParser.h"
+
+#include "RtspMediaConstants.h"
+
+bool FRtpParser::Parse(const TArrayView<const uint8> InData, FRtpHeader& OutHeader, TArrayView<const uint8>& OutPayload)
+{
+	const int32 DataSize = InData.Num();
+
+	constexpr int32 MinimumHeaderSize = 12;
+
+	// 12 Bytes is the minimum header size
+	if (DataSize < MinimumHeaderSize)
+	{
+		UE_LOG(LogRtspMedia, Error, TEXT("RTP packet size is below the minimum header size: %d"), MinimumHeaderSize);
+		return false;
+	}
+
+	FRtpHeader Header;
+
+	const uint8* Data = InData.GetData();
+
+	// Byte 0 - Version, bPadding, bExtension, CsrcCount
+	const uint8 Byte0 = Data[0];
+	Header.Version = (Byte0 >> 6) & 0x03;
+	Header.bPadding = (Byte0 >> 5) & 0x01;
+	Header.bExtension = (Byte0 >> 4) & 0x01;
+	Header.CsrcCount = Byte0 & 0x0F;
+
+	if (Header.Version != 2)
+	{
+		UE_LOG(LogRtspMedia, Error, TEXT("RTP packet version must be '2' instead we received '%d'"), Header.Version);
+		return false;
+	}
+
+	// Byte 1 - bMarker and PayloadType
+	const uint8 Byte1 = Data[1];
+	Header.bMarker = (Byte1 >> 7) & 0x01;
+	Header.PayloadType = Byte1 & 0x7F;
+
+	// All multi-byte values are in network byte order
+	// Bytes 2 - 3: Sequence Number
+	Header.SequenceNumber = (static_cast<uint16>(Data[2]) << 8) | Data[3];
+
+	// Bytes 4 - 7: Timestamp
+	const uint32 TimestampByte0 = static_cast<uint32>(Data[4]) << 24;
+	const uint32 TimestampByte1 = static_cast<uint32>(Data[5]) << 16;
+	const uint32 TimestampByte2 = static_cast<uint32>(Data[6]) << 8;
+	const uint32 TimestampByte3 = static_cast<uint32>(Data[7]);
+	Header.Timestamp = TimestampByte0 | TimestampByte1 | TimestampByte2 | TimestampByte3;
+
+	// Bytes 8 - 11: Ssrc
+	const uint32 SsrcByte0 = static_cast<uint32>(Data[8]) << 24;
+	const uint32 SsrcByte1 = static_cast<uint32>(Data[9]) << 16;
+	const uint32 SsrcByte2 = static_cast<uint32>(Data[10]) << 8;
+	const uint32 SsrcByte3 = static_cast<uint32>(Data[11]);
+	Header.Ssrc = SsrcByte0 | SsrcByte1 | SsrcByte2 | SsrcByte3;
+
+	// We're not interested in the contributing source identifiers so we skip over those.
+	// Each identifier is 4 bytes in size.
+	constexpr int32 ContributingSourceIdentifierSize = 4;
+	int32 HeaderSize = MinimumHeaderSize + (Header.CsrcCount * ContributingSourceIdentifierSize);
+
+	if (Header.bExtension)
+	{
+		// Extension Header contains:
+		// - Profile ID (2 bytes)
+		// - Length (2 bytes) - Number of 32-bit words
+		// - Extension size (MinimumExtensionHeaderSize + (Length * 4))
+		// We don't use the contents of the extension header, but we need to figure
+		// out it's size in order to determine where the payload starts.
+
+		constexpr int32 MinimumExtensionHeaderSize = 4;
+		// InData must contain at least the header size and the minimum extension header size bytes
+		const int32 MinimumSizeWithExtensionHeader = HeaderSize + MinimumExtensionHeaderSize;
+		if (DataSize < MinimumSizeWithExtensionHeader)
+		{
+			// If not this is a malformed packet
+			UE_LOG(LogRtspMedia, Warning, TEXT("RTP packet (size: %d) containing extension header below minimum size: %d"), DataSize, MinimumSizeWithExtensionHeader);
+			return false; 
+		}
+
+		// We skip over Profile ID and go straight to Length (HeaderSize + 2)
+
+		// ExtensionLength is the number of 32-bit words so the number of bytes is (ExtensionLength * 4)
+		const uint16 ExtensionLengthByte0 = static_cast<uint16>(Data[HeaderSize + 2]) << 8;
+		const uint16 ExtensionLengthByte1 = static_cast<uint16>(Data[HeaderSize + 3]);
+		const uint16 ExtensionLength = ExtensionLengthByte0 | ExtensionLengthByte1;
+
+		// Increase the HeaderSize by the 4 byte header + (ExtensionLength * 4) bytes
+		HeaderSize += MinimumExtensionHeaderSize + (static_cast<int32>(ExtensionLength) * 4);
+	}
+
+	// Final check that data contains at least the size of the header
+	if (DataSize < HeaderSize)
+	{
+		// Malformed packet
+		UE_LOG(LogRtspMedia, Warning, TEXT("RTP packet (size: %d) is below expected header size: %d"), DataSize, HeaderSize);
+		return false;
+	}
+
+	OutHeader = Header;
+	OutPayload = InData.Slice(HeaderSize, DataSize - HeaderSize);
+
+	if (Header.bPadding)
+	{
+		if (OutPayload.IsEmpty())
+		{
+			UE_LOG(LogRtspMedia, Warning, TEXT("RTP packet padding bit is set but the payload is empty. Packet is malformed"));
+			return false;
+		}
+
+		const int32 PayloadSize = OutPayload.Num();
+		const uint8 PaddingCount = OutPayload.Last();
+		if (PaddingCount > 0 && PaddingCount <= PayloadSize)
+		{
+			OutPayload = OutPayload.Slice(0, PayloadSize - PaddingCount);
+		}
+		else 
+		{
+			UE_LOG(LogRtspMedia, Warning, TEXT("Padding count invalid within RTP packet: %d. Payload size is: %d"), PaddingCount, PayloadSize);
+			return false;
+		}
+	}
+
+	return true;
+}
\ No newline at end of file
diff --git a/Engine/Plugins/Media/RTSPMedia/Source/RTSPMedia/Private/RTP/RtpParser.h b/Engine/Plugins/Media/RTSPMedia/Source/RTSPMedia/Private/RTP/RtpParser.h
new file mode 100644
--- /dev/null
+++ b/Engine/Plugins/Media/RTSPMedia/Source/RTSPMedia/Private/RTP/RtpParser.h
@@ -0,0 +1,13 @@
+// Copyright Epic Games, Inc. All Rights Reserved.
+
+#pragma once
+
+#include "CoreMinimal.h"
+
+#include "RtpPacket.h"
+
+class FRtpParser
+{
+public:
+	static bool Parse(const TArrayView<const uint8> InData, FRtpHeader& OutHeader, TArrayView<const uint8>& OutPayload);
+};
diff --git a/Engine/Plugins/Media/RTSPMedia/Source/RTSPMedia/Private/RTSP/RtspClient.cpp b/Engine/Plugins/Media/RTSPMedia/Source/RTSPMedia/Private/RTSP/RtspClient.cpp
new file mode 100644
--- /dev/null
+++ b/Engine/Plugins/Media/RTSPMedia/Source/RTSPMedia/Private/RTSP/RtspClient.cpp
@@ -0,0 +1,1001 @@
+// Copyright Epic Games, Inc. All Rights Reserved.
+
+#include "RtspClient.h"
+
+#include "RtspMediaConstants.h"
+
+#include "RTP/RtpParser.h"
+#include "RTP/RtpPacket.h"
+
+#include "RTSP/RtspMessage.h"
+#include "RTSP/RtspTransportConfiguration.h"
+
+#include "SDP/SdpParser.h"
+#include "SDP/SdpSession.h"
+
+#include "Async/Async.h"
+#include "HAL/PlatformProcess.h"
+#include "SocketSubsystem.h"
+#include "Sockets.h"
+#include "Common/TcpSocketBuilder.h"
+#include "RTP/RtpDecoder.h"
+
+FRtspClient::~FRtspClient()
+{
+	Disconnect();
+}
+
+bool FRtspClient::Initialize(const FRtspClientConfiguration& InConfiguration)
+{
+	if (GetState() != ERtspClientState::Disconnected)
+	{
+		UE_LOG(LogRtspMedia, Warning, TEXT("Can only initialize when disconnected"));
+		return false;
+	}
+
+	// Validate URL
+	const FString Url = InConfiguration.Url;
+	if (!ParseUrl(Url, Host, Port, Path))
+	{
+		UE_LOG(LogRtspMedia, Error, TEXT("Invalid RTSP URL: %s"), *Url);
+		return false;
+	}
+
+	// Store configuration
+	Configuration = InConfiguration;
+
+	FString BaseUrl = FString::Printf(TEXT("rtsp://%s:%d"), *Host, Port);
+	FullUrl = Path.IsEmpty() ? BaseUrl : (BaseUrl / Path);
+
+	// Start network thread
+	UE_LOG(LogRtspMedia, Log, TEXT("Connecting to %s"), *FullUrl);
+	SetState(ERtspClientState::Connecting);
+
+	bStopping.store(false);
+
+	Thread.Reset(FRunnableThread::Create(this, TEXT("RtspClientThread"), 0, TPri_Normal));
+
+	if (!Thread.IsValid())
+	{
+		UE_LOG(LogRtspMedia, Error, TEXT("Failed to create network thread"));
+		SetState(ERtspClientState::Error);
+		return false;
+	}
+
+	return true;
+}
+
+void FRtspClient::Disconnect()
+{
+	if (GetState() == ERtspClientState::Disconnected)
+	{
+		return;
+	}
+
+	UE_LOG(LogRtspMedia, Log, TEXT("Disconnecting from %s"), *FullUrl);
+
+	bStopping.store(true);
+
+	if (Thread.IsValid())
+	{
+		Thread->Kill(true);
+		Thread.Reset();
+	}
+}
+
+bool FRtspClient::InitializeJitterBuffer(const uint32 InClockRate)
+{
+	JitterBuffer = MakeUnique<FRtpJitterBuffer>();
+	if (!JitterBuffer->Initialize(Configuration.JitterBufferDepthMs, InClockRate))
+	{
+		JitterBuffer.Reset();
+		return false;
+	}
+	return true;
+}
+
+ERtspClientState FRtspClient::GetState() const
+{
+	// State can be read from both the worker thread and the game thread
+	FScopeLock Lock(&StateLock);
+	return State;
+}
+
+bool FRtspClient::Init()
+{
+	return true;
+}
+
+uint32 FRtspClient::Run()
+{
+	Socket = TUniquePtr<FSocket>(
+		FTcpSocketBuilder(TEXT("RtspSocket"))
+			.AsNonBlocking()
+			.AsReusable()
+			.WithReceiveBufferSize(Configuration.SocketBufferSizeBytes)
+			.WithSendBufferSize(Configuration.SocketBufferSizeBytes)
+			.Build()
+	);
+
+	if (!Socket.IsValid())
+	{
+		UE_LOG(LogRtspMedia, Error, TEXT("Failed to create socket"));
+		return HandleRunError();
+	}
+
+	// Configure the socket for no delay, which disables Nagle's algorithm.
+	// Only really affects outgoing packets, but should speed up the initial handshake.
+	Socket->SetNoDelay(true);
+
+	ISocketSubsystem* SocketSubsystem = ISocketSubsystem::Get();
+	if (SocketSubsystem == nullptr)
+	{
+		UE_LOG(LogRtspMedia, Error, TEXT("Failed to get socket subsystem"));
+		return HandleRunError();
+	}
+
+	FAddressInfoResult AddressInfoResult = SocketSubsystem->GetAddressInfo(*Host, nullptr, EAddressInfoFlags::Default, NAME_None);
+	if (AddressInfoResult.ReturnCode != SE_NO_ERROR || AddressInfoResult.Results.IsEmpty())
+	{
+		UE_LOG(LogRtspMedia, Warning, TEXT("Failed to resolve host: %s"), *Host);
+		return HandleRunError();
+	}
+	TSharedRef<FInternetAddr> ServerAddress = AddressInfoResult.Results[0].Address;
+
+	ServerAddress->SetPort(Port);
+
+	// Connect to Server
+	// We're using a non-blocking socket.
+	// FSocket::Connect() returns true for EINPROGRESS/EWOULDBLOCK conditions, which is expected
+	// for non-blocking sockets. A false return indicates a genuine failure.
+	// (e.g. bad address, network unreachable). We then poll with Wait(WaitForWrite) and confirm via
+	// GetConnectionState() to determine the actual connection outcome.
+	const bool bConnected = Socket->Connect(*ServerAddress);
+	if (!bConnected)
+	{
+		UE_LOG(LogRtspMedia, Warning, TEXT("Failed to connect socket for host: %s"), *Host);
+		return HandleRunError();
+	}
+
+	// Wait for write
+	if (!Socket->Wait(ESocketWaitConditions::WaitForWrite, FTimespan::FromSeconds(5.0)))
+	{
+		UE_LOG(LogRtspMedia, Warning, TEXT("Connection timed out for host: %s"), *Host);
+		return HandleRunError();
+	}
+
+	if (Socket->GetConnectionState() != SCS_Connected)
+	{
+		UE_LOG(LogRtspMedia, Warning, TEXT("Connection failed for host: %s"), *Host);
+		return HandleRunError();
+	}
+
+	UE_LOG(LogRtspMedia, Log, TEXT("Connected to host: %s"), *Host);
+	SetState(ERtspClientState::Connected);
+
+	ReceiveBuffer.SetNumUninitialized(4096);
+
+	// Main loop
+
+	while (!bStopping.load())
+	{
+		ProcessState();
+
+		if (GetState() == ERtspClientState::Error)
+		{
+			break;
+		}
+
+		if (Socket->Wait(ESocketWaitConditions::WaitForRead, FTimespan::FromMilliseconds(5)))
+		{
+			if (!ProcessIncomingData())
+			{
+				SetState(ERtspClientState::Error);
+				continue;
+			}
+		}
+
+		// Check for waiting packets in the jitter buffer
+		if (JitterBuffer.IsValid())
+		{
+			FRtpPacket BufferedPacket;
+			while (JitterBuffer->DequeueIfReady(BufferedPacket))
+			{
+				// Note: Executed on the client worker thread to avoid a hop to the game thread
+				OnRtpPacketReceived.ExecuteIfBound(MoveTemp(BufferedPacket));
+			}
+		}
+
+		// Check for expired requests
+		TArray<FRtspPendingRequest> ExpiredRequests = RequestTracker.GetExpiredRequests(Configuration.RequestTimeoutSeconds);
+		for (const FRtspPendingRequest& ExpiredRequest : ExpiredRequests)
+		{
+			const ERtspRequestMethod Method = ExpiredRequest.Request.GetRequestMethod().Get(ERtspRequestMethod::Unknown);
+			FString MethodString = FRtspMessage::RequestMethodToString(Method);
+			UE_LOG(LogRtspMedia, Error, TEXT("%s request timed out"), *MethodString);
+			SetState(ERtspClientState::Error);
+			break;
+		}
+
+		if (GetState() == ERtspClientState::Playing)
+		{
+			// Send keep alive
+			const double Now = FPlatformTime::Seconds();
+			if (Now - LastKeepAlive > KeepAliveInterval)
+			{
+				if (!SendKeepAlive())
+				{
+					SetState(ERtspClientState::Error);
+				}
+			}
+		}
+		else 
+		{
+			FPlatformProcess::Sleep(0.001f);
+		}
+	}
+	
+	TeardownConnection();
+
+	SetState(ERtspClientState::Disconnected);
+
+	UE_LOG(LogRtspMedia, Verbose, TEXT("Network thread exiting"));
+	
+	return 0;
+}
+
+void FRtspClient::Stop()
+{
+	bStopping.store(true);
+}
+
+const TCHAR* FRtspClient::StateToString(ERtspClientState InState)
+{
+	switch (InState)
+	{
+		case ERtspClientState::Disconnected: return TEXT("Disconnected");
+		case ERtspClientState::Connecting: return TEXT("Connecting");
+		case ERtspClientState::Connected: return TEXT("Connected");
+		case ERtspClientState::Ready: return TEXT("Ready");
+		case ERtspClientState::Described: return TEXT("Described");
+		case ERtspClientState::Setup: return TEXT("Setup");
+		case ERtspClientState::Playing: return TEXT("Playing");
+		case ERtspClientState::Error: return TEXT("Error");
+		default: return TEXT("Unknown");
+	}
+}
+
+void FRtspClient::SetState(ERtspClientState InNewState)
+{
+	ERtspClientState OldState;
+	{
+		FScopeLock Lock(&StateLock);
+
+		if (State == InNewState)
+		{
+			return;
+		}
+
+		OldState = State;
+		State = InNewState;
+	}
+
+	UE_LOG(LogRtspMedia, Verbose, TEXT("State changed: %s -> %s"), StateToString(OldState), StateToString(InNewState));
+
+	AsyncTask(ENamedThreads::GameThread, [Delegate = OnStateChanged, InNewState, OldState]()
+	{
+		Delegate.ExecuteIfBound(InNewState, OldState);
+	});
+}
+
+void FRtspClient::ProcessState()
+{
+	// Check if we have any requests in flight, if so wait until the response has been
+	// handled or times out before performing the next state transition.
+	if (RequestTracker.HasPendingRequests())
+	{
+		return;
+	}
+	
+	ERtspClientState CurrentState = GetState();
+	switch (CurrentState) 
+	{
+		// Just connected. Send OPTIONS to query server capabilities.
+		case ERtspClientState::Connected:
+			if (!SendOptions())
+			{
+				SetState(ERtspClientState::Error);
+			}
+			break;
+		// OPTIONS succeeded. Send DESCRIBE to get SDP (Session Description Protocol) description data
+		case ERtspClientState::Ready:
+			if (!SendDescribe())
+			{
+				SetState(ERtspClientState::Error);
+			}
+			break;
+		// DESCRIBE succeeded. Send SETUP for each track
+		// Basic initial implementation, just request the first video track
+		case ERtspClientState::Described:
+			if (SdpSession.IsValid() && SdpSession->VideoTracks.Num() > 0)
+			{
+				const FSdpMediaTrack& Track = SdpSession->VideoTracks[0];
+				if (!SendSetup(Track.ControlUrl, 0))
+				{
+					SetState(ERtspClientState::Error);
+				}
+			}
+			else 
+			{
+				UE_LOG(LogRtspMedia, Error, TEXT("No video tracks in session description"));
+				SetState(ERtspClientState::Error);
+			}
+			break;
+		// SETUP succeeded. Send PLAY to start streaming.
+		case ERtspClientState::Setup:
+			if (!SendPlay())
+			{
+				SetState(ERtspClientState::Error);
+			}
+			break;
+		case ERtspClientState::Playing:
+		case ERtspClientState::Error:
+		case ERtspClientState::Disconnected:
+		case ERtspClientState::Connecting:
+			// No state transition to make
+			break;
+	}
+}
+
+bool FRtspClient::ProcessIncomingData()
+{
+	int32 BytesRead = 0;
+	const bool bSuccess = Socket->Recv(ReceiveBuffer.GetData(), ReceiveBuffer.Num(), BytesRead);
+
+	if (!bSuccess)
+	{
+		UE_LOG(LogRtspMedia, Warning, TEXT("Failed to receive data from socket"));
+		return false;
+	}
+
+	if (BytesRead == 0)
+	{
+		UE_LOG(LogRtspMedia, Warning, TEXT("Connection closed by server"));
+		return false;
+	}
+
+	// 64 KB Max Process Buffer Size
+	constexpr int32 MaxProcessBufferSize = 64 * 1024;
+	if (ProcessBuffer.Num() + BytesRead > MaxProcessBufferSize)
+	{
+		UE_LOG(LogRtspMedia, Error, TEXT("Max RTSP process buffer size (%d) exceeded."), MaxProcessBufferSize);
+		ProcessBuffer.Reset();
+		return false;
+	}
+
+	ProcessBuffer.Append(ReceiveBuffer.GetData(), BytesRead);
+	
+	while (!ProcessBuffer.IsEmpty())
+	{
+		const int32 ProcessBufferSize = ProcessBuffer.Num();
+		
+		if (ProcessBuffer[0]  == '$')
+		{
+			ProcessInterleavedPacket();
+		}
+		else
+		{
+			if (!ProcessRtspMessage())
+			{
+				return false;
+			}
+		}
+
+		// If the buffer size has not reduced then we have incomplete RTP or RTSP message data
+		// and need to read more in the next iteration.
+		if (ProcessBufferSize == ProcessBuffer.Num())
+		{
+			break;
+		}
+	}
+
+	return true;
+}
+
+void FRtspClient::ProcessInterleavedPacket()
+{
+	constexpr int32 HeaderLength = 4;
+	if (ProcessBuffer.Num() < HeaderLength)
+	{
+		// Incomplete data
+		// We want the run loop to continue and read more data.
+		return;
+	}
+
+	// Index 0 is the interleaved packet indicator '$'
+
+	// Channel lies at index 1
+	const uint8 Channel = ProcessBuffer[1];
+
+	// Packet length is at index 2
+	// Data is big endian (Network byte order)
+	const uint16 InterleavedPayloadLengthByte0 = static_cast<uint16>(ProcessBuffer[2]) << 8;
+	const uint16 InterleavedPayloadLengthByte1 = static_cast<uint16>(ProcessBuffer[3]);
+	const uint16 InterleavedPayloadLength = InterleavedPayloadLengthByte0 | InterleavedPayloadLengthByte1;
+
+	if (ProcessBuffer.Num() < HeaderLength + InterleavedPayloadLength)
+	{
+		// Incomplete data
+		// We want the run loop to continue and read more data.
+		return;
+	}
+	
+	// TODO: Map channel and track information when audio support is added.
+	// Different channels represent different streams. Typically:
+	// 0 - Video RTP
+	// 1 - Video RTCP
+	// 2 - Audio RTP
+	// 3 - Audio RTCP
+	// We only support the video RTP channel
+	if (Channel == 0)
+	{
+		const TArrayView<const uint8> InterleavedPayload(ProcessBuffer.GetData() + HeaderLength, InterleavedPayloadLength);
+		
+		FRtpHeader Header;
+		TArrayView<const uint8> RtpPayload;
+		if (FRtpParser::Parse(InterleavedPayload, Header, RtpPayload))
+		{
+			if (JitterBuffer.IsValid())
+			{
+				FRtpPacket Packet{Header, TArray<uint8>(RtpPayload)};
+				JitterBuffer->EnqueuePacket(MoveTemp(Packet));
+			}
+			else
+			{
+				UE_LOG(LogRtspMedia, Verbose, TEXT("Jitter buffer not initialized when processing interleaved RTP packet"));
+			}
+		}
+		else
+		{
+			UE_LOG(LogRtspMedia, Warning, TEXT("Failed to parse RTP packet"));
+		}
+	}
+	else
+	{
+		UE_LOG(LogRtspMedia, VeryVerbose, TEXT("Dropping unsupported packet on channel %d"), Channel);
+	}
+
+	// Once enqueued and copied to the jitter buffer remove those bytes from our buffer
+	// Or, if we're dropping the packet, remove it from the buffer.
+	ProcessBuffer.RemoveAt(0, HeaderLength + InterleavedPayloadLength);
+}
+
+bool FRtspClient::ProcessRtspMessage()
+{
+	FRtspMessage Message;
+	int32 BytesConsumed = 0;
+
+	switch (FRtspMessage::Parse(ProcessBuffer, Message, BytesConsumed))
+	{
+		// Cannot parse the message yet, not enough data. Try again on the next iteration.
+		case ERtspMessageParseResult::Incomplete:
+			return true;
+		// An error occurred during parsing, abort.
+		case ERtspMessageParseResult::Error:
+			return false;
+		// Message was successfully parsed, continue.
+		case ERtspMessageParseResult::Complete:
+			break;
+	}
+	
+	ProcessBuffer.RemoveAt(0, BytesConsumed);
+	
+	switch (Message.GetMessageType())
+	{
+		case ERtspMessageType::Unknown:
+			UE_LOG(LogRtspMedia, Error, TEXT("Received RTSP message of an unknown type"));
+			return false;
+		case ERtspMessageType::Request:
+			return HandleRequest(Message);
+		case ERtspMessageType::Response:
+			return HandleResponse(Message);
+	}
+
+	return true;
+}
+
+bool FRtspClient::HandleResponse(const FRtspMessage& InResponse)
+{
+	if (InResponse.GetMessageType() != ERtspMessageType::Response)
+	{
+		UE_LOG(LogRtspMedia, Error, TEXT("Handle response called with non-response RTSP message"));
+		return false;
+	}
+
+	UE_LOG(LogRtspMedia, Verbose, TEXT("Received RTSP response:\n%s-----\n"), *InResponse.ToString());
+	
+	const TOptional<int32> ResponseCommandId = InResponse.GetCommandId();
+
+	if (!ResponseCommandId.IsSet())
+	{
+		UE_LOG(LogRtspMedia, Warning, TEXT("Failed to process RTSP response with missing command ID"))
+		return false;
+	}
+
+	FRtspPendingRequest PendingRequest;
+	if (!RequestTracker.TakePendingRequest(ResponseCommandId.GetValue(), PendingRequest))
+	{
+		UE_LOG(LogRtspMedia, Warning, TEXT("Received RTSP response with unknown command ID: %d"), ResponseCommandId.GetValue());
+		// Ignore this unexpected response, don't abort processing.
+		return true;
+	}
+
+	TOptional<int32> StatusCode = InResponse.GetResponseStatusCode();
+	int32 StatusCodeValue = -1;
+	bool bSuccess = false;
+	if (StatusCode.IsSet())
+	{
+		StatusCodeValue = StatusCode.GetValue();
+		bSuccess = StatusCodeValue >= 200 && StatusCodeValue < 300;
+	}
+
+	TOptional<FString> Reason = InResponse.GetResponseReason();
+	FString ReasonString;
+	if (Reason.IsSet())
+	{
+		ReasonString = Reason.GetValue();
+	}
+
+	FRtspMessage Request = PendingRequest.Request;
+	
+	TOptional<ERtspRequestMethod> Method = Request.GetRequestMethod();
+	if (!Method.IsSet())
+	{
+		// If our pending request didn't have a method set then we can't process it.
+		UE_LOG(LogRtspMedia, Warning, TEXT("Pending request doesn't contain a METHOD"));
+		return true;
+	}
+	
+	if (!bSuccess)
+	{
+		FString MethodString = FRtspMessage::RequestMethodToString(Method.GetValue());
+		UE_LOG(LogRtspMedia, Warning, TEXT("%s request failed: %d %s"), *MethodString, StatusCodeValue, *ReasonString);
+	}
+
+	switch (Method.GetValue())
+	{
+		case ERtspRequestMethod::Options:
+			return HandleOptionsResponse(StatusCodeValue, bSuccess, Request);
+		case ERtspRequestMethod::Describe:
+			return HandleDescribeResponse(bSuccess, InResponse, Request);
+		case ERtspRequestMethod::Setup:
+			return HandleSetupResponse(bSuccess, InResponse);
+		case ERtspRequestMethod::Play:
+			if (!bSuccess)
+			{
+				return false;
+			}
+			UE_LOG(LogRtspMedia, Log, TEXT("Playback started"));
+			// Prevent unnecessary keep-alive message when playing is started
+			LastKeepAlive = FPlatformTime::Seconds();
+			SetState(ERtspClientState::Playing);
+			break;
+		default:
+			break;
+	}
+
+	return true;
+}
+
+bool FRtspClient::HandleOptionsResponse(const int32 InStatusCodeValue, const bool bInSuccess, const FRtspMessage& InOriginalRequest)
+{
+	// The only differentiator between an initial setup OPTIONS request and
+	// the keep-alive OPTIONS request is the presence of the sessionId.
+	// This is not present before we have set up a session.
+	if (InOriginalRequest.GetSessionId().IsSet())
+	{
+		if (bInSuccess)
+		{
+			return true;
+		}
+
+		if (InStatusCodeValue == 454)
+		{
+			UE_LOG(LogRtspMedia, Error, TEXT("Keep-alive failed. Session not found."));
+		}
+	}
+	// Otherwise we're looking at the initial OPTIONS request response when setting up a session.
+	else
+	{
+		if (bInSuccess)
+		{
+			SetState(ERtspClientState::Ready);
+			return true;
+		}
+	}
+	
+	return false;
+}
+
+bool FRtspClient::HandleDescribeResponse(const bool bInSuccess, const FRtspMessage& InResponse, const FRtspMessage& InOriginalRequest)
+{
+	if (!bInSuccess)
+	{
+		return false;
+	}
+	
+	TOptional<FString> ContentBaseUrl = InResponse.GetContentBaseUrl();
+	FString ContentBaseUrlString;
+	if (ContentBaseUrl.IsSet())
+	{
+		ContentBaseUrlString = ContentBaseUrl.GetValue();
+	}
+	else
+	{
+		// Fall back to the originally requested URL
+		ContentBaseUrlString = InOriginalRequest.GetRequestUrl().Get(FString());
+	}
+	
+	// Parse SDP Session
+	SdpSession = FSdpParser::Parse(InResponse.GetBody(), ContentBaseUrlString);
+	if (!SdpSession.IsValid())
+	{
+		UE_LOG(LogRtspMedia, Error, TEXT("Failed to parse SDP session"));
+		return false;
+	}
+
+	const FSdpMediaTrack* VideoTrack = SdpSession->GetFirstVideoTrack();
+	if (VideoTrack == nullptr)
+	{
+		UE_LOG(LogRtspMedia, Error, TEXT("No video track found within SDP session"));
+		return false;
+	}
+	
+	// Initialise jitter buffer in the RTSP client
+	if (!InitializeJitterBuffer(VideoTrack->ClockRate))
+	{
+		return false;
+	}
+
+	AsyncTask(ENamedThreads::GameThread, [Delegate = OnSdpSessionReceived, SdpSession = *SdpSession]()
+	{
+		Delegate.ExecuteIfBound(SdpSession);
+	});
+	
+	SetState(ERtspClientState::Described);
+	return true;
+}
+
+bool FRtspClient::HandleSetupResponse(const bool bInSuccess, const FRtspMessage& InResponse)
+{
+	if (!bInSuccess)
+	{
+		return false;
+	}
+	
+	if (SessionId.IsEmpty())
+    {
+    	TOptional<FString> ResponseSessionId = InResponse.GetSessionId();
+    	if (!ResponseSessionId.IsSet())
+    	{
+    		UE_LOG(LogRtspMedia, Error, TEXT("No session ID in SETUP response"));
+    		return false;
+    	}
+
+		SessionId = ResponseSessionId.GetValue();
+    	UE_LOG(LogRtspMedia, Verbose, TEXT("Session ID: %s"), *SessionId);
+
+    	TOptional<int32> SessionTimeout = InResponse.GetSessionTimeout();
+    	if (SessionTimeout.IsSet())
+    	{
+    		SessionTimeoutSeconds = SessionTimeout.GetValue();
+    		KeepAliveInterval = static_cast<double>(SessionTimeoutSeconds) / 2.0;
+    		UE_LOG(LogRtspMedia, Verbose, TEXT("SETUP response provided session timeout of %d seconds"), SessionTimeoutSeconds);
+    	}
+    	else
+    	{
+    		UE_LOG(LogRtspMedia, Verbose, TEXT("No timeout specified in SETUP response. Assuming %d seconds"), SessionTimeoutSeconds);
+    	}
+    	UE_LOG(LogRtspMedia, Verbose, TEXT("Keep-alive interval set to %.2f seconds"), KeepAliveInterval);
+    }
+    
+    SetState(ERtspClientState::Setup);
+	return true;
+}
+
+bool FRtspClient::HandleRequest(const FRtspMessage& InRequest)
+{
+	if (InRequest.GetMessageType() != ERtspMessageType::Request)
+	{
+		return false;
+	}
+
+	UE_LOG(LogRtspMedia, Verbose, TEXT("Received RTSP request:\n%s-----\n"), *InRequest.ToString());
+
+	const ERtspRequestMethod Method = InRequest.GetRequestMethod().Get(ERtspRequestMethod::Unknown);
+	const TOptional<int32> RequestCommandId = InRequest.GetCommandId();
+
+	if (!RequestCommandId.IsSet())
+	{
+		UE_LOG(LogRtspMedia, Warning, TEXT("Cannot process RTSP request with missing command ID"));
+		return false;
+	}
+
+	const FString RequestMethodString = FRtspMessage::RequestMethodToString(Method);
+
+	UE_LOG(LogRtspMedia, Verbose, TEXT("Received RTSP %s request from server with command ID: %d"), *RequestMethodString, RequestCommandId.GetValue());
+
+	FRtspMessage Response;
+
+	switch (Method)
+	{
+	case ERtspRequestMethod::GetParameter:
+		if (InRequest.GetBody().IsEmpty())
+		{
+			// RFC 2326 Section 10.8
+			// https://datatracker.ietf.org/doc/html/rfc2326#section-10.8
+			// Empty body GET_PARAMETER requests are the per-spec way to perform keep alive messaging.
+			// Some servers will use this to check if client is still active and interested.
+			Response = FRtspMessage::Response(200, TEXT("OK"), RequestCommandId.GetValue());
+		}
+		else
+		{
+			// We don't support supplying any actual parameters
+			Response = FRtspMessage::Response(451, TEXT("Parameter Not Understood"), RequestCommandId.GetValue());
+		}
+		break;
+	case ERtspRequestMethod::Options:
+		// There's a convention that some servers might follow of using OPTIONS requests as keep-alive messages.
+		// The OPTIONS response should contain the requests that we can handle.
+		// That handling is in this method, so all we support is GET_PARAMETER and OPTIONS
+		Response = FRtspMessage::Response(200, TEXT("OK"), RequestCommandId.GetValue());
+		Response.SetHeader(TEXT("Public"), TEXT("GET_PARAMETER, OPTIONS"));
+		break;
+	default:
+		UE_LOG(LogRtspMedia, Verbose, TEXT("Received unsupported RTSP request from server: %s"), *RequestMethodString);
+		Response = FRtspMessage::Response(501, TEXT("Not Implemented"), RequestCommandId.GetValue());
+	}
+
+	// If the sender provides a session ID then we'll mirror this back in the response
+	TOptional<FString> RequestSessionId = InRequest.GetSessionId();
+	if (RequestSessionId.IsSet())
+	{
+		FString RequestSessionIdString = RequestSessionId.GetValue();
+		bool bSessionNotFound = false;
+		
+		// We don't have a session ID yet, why did the request contain one?
+		if (SessionId.IsEmpty())
+		{
+			UE_LOG(LogRtspMedia, Warning, TEXT("RTSP request contains a session ID (%s), but setup is yet to occur."), *RequestSessionIdString);
+			bSessionNotFound = true;
+		}
+		// Check that this matches our known SessionId
+		// If it doesn't match something has gone wrong on the server side
+		else if (RequestSessionIdString != SessionId)
+		{
+			UE_LOG(LogRtspMedia, Warning, TEXT("Received request for unknown session ID: %s Expected: %s"), *RequestSessionIdString, *SessionId);
+			bSessionNotFound = true;
+		}
+
+		// If session not found replace the pending response with 'Session Not Found'
+		if (bSessionNotFound)
+		{
+			Response = FRtspMessage::Response(454, TEXT("Session Not Found"), RequestCommandId.GetValue());
+		}
+		else
+		{
+			Response.SetSession(RequestSessionIdString);	
+		}
+	}
+
+	return SendMessage(Response);
+}
+
+bool FRtspClient::SendOptions()
+{
+	UE_LOG(LogRtspMedia, Verbose, TEXT("Sending OPTIONS"));
+	const FRtspMessage Request = FRtspMessage::Request(ERtspRequestMethod::Options, FullUrl, ++CommandId);
+	return SendRequest(Request);
+}
+
+bool FRtspClient::SendDescribe()
+{
+	UE_LOG(LogRtspMedia, Verbose, TEXT("Sending DESCRIBE"));
+	FRtspMessage Request = FRtspMessage::Request(ERtspRequestMethod::Describe, FullUrl, ++CommandId);
+	Request.SetAccept(TEXT("application/sdp"));
+	return SendRequest(Request);
+}
+
+bool FRtspClient::SendSetup(const FString& InControlUrl, const int32 InTrackIndex)
+{
+	UE_LOG(LogRtspMedia, Verbose, TEXT("Sending SETUP for track %d: %s"), InTrackIndex, *InControlUrl);
+
+	const FString ResolvedControlUrl = SdpSession->ResolveControlUrl(InControlUrl);
+	FRtspMessage Request = FRtspMessage::Request(ERtspRequestMethod::Setup, ResolvedControlUrl, ++CommandId);
+
+	if (!SessionId.IsEmpty())
+	{
+		Request.SetSession(SessionId);
+	}
+
+	FRtspTransportConfiguration::TransportProtocol TransportProtocol = Configuration.TransportProtocol;
+	check(TransportProtocol == FRtspTransportConfiguration::TransportProtocol::TCP || TransportProtocol == FRtspTransportConfiguration::TransportProtocol::UDP);
+
+	// Each track needs two channels/ports (one for RTP and another for RTCP, even if we're not using the latter)
+	// Convention is even number for RTP, odd number for RTCP.
+	if (TransportProtocol == FRtspTransportConfiguration::TransportProtocol::TCP)
+	{
+		const int32 RtpChannel = InTrackIndex * 2;
+		const int32 RtcpChannel = RtpChannel + 1;
+		Request.SetTransport(FRtspTransportConfiguration::TcpInterleaved(RtpChannel, RtcpChannel));
+	}
+	else if (TransportProtocol == FRtspTransportConfiguration::TransportProtocol::UDP)
+	{
+		// Using a fixed based port offset for now.
+		// TODO: Dynamic port allocation
+		const int32 RtpPort = 50000 + InTrackIndex * 2;
+		const int32 RtcpPort = RtpPort + 1;
+		Request.SetTransport(FRtspTransportConfiguration::Udp(RtpPort, RtcpPort));
+	}
+
+	return SendRequest(Request);
+}
+
+bool FRtspClient::SendPlay()
+{
+	UE_LOG(LogRtspMedia, Verbose, TEXT("Sending PLAY"));
+	FRtspMessage Request = FRtspMessage::Request(ERtspRequestMethod::Play, FullUrl, ++CommandId);
+	Request.SetSession(SessionId);
+	return SendRequest(Request);
+}
+
+bool FRtspClient::SendTeardown()
+{
+	UE_LOG(LogRtspMedia, Verbose, TEXT("Sending TEARDOWN"));
+	FRtspMessage Request = FRtspMessage::Request(ERtspRequestMethod::Teardown, FullUrl, ++CommandId);
+	Request.SetSession(SessionId);
+	// Using SendMessage here to avoid adding the request to the tracker.
+	// We're usually shutting down at this point so we don't need to wait for a response.
+	return SendMessage(Request);
+}
+
+// We lean on the common convention to use an OPTIONS request containing the session ID as a keep-alive message for the RTSP connection.
+// Officially an empty GET_PARAMETER request should be used for keep-alive, but as it's an optional method it's often absent.
+bool FRtspClient::SendKeepAlive()
+{
+	UE_LOG(LogRtspMedia, Verbose, TEXT("Sending keep-alive (OPTIONS)"));
+	FRtspMessage Request = FRtspMessage::Request(ERtspRequestMethod::Options, FullUrl, ++CommandId);
+	Request.SetSession(SessionId);
+	LastKeepAlive = FPlatformTime::Seconds();
+	return SendRequest(Request);
+}
+
+bool FRtspClient::SendMessage(const FRtspMessage& InMessage)
+{
+	const FString MessageString = InMessage.ToString();
+	UE_LOG(LogRtspMedia, Verbose, TEXT("Sending RTSP message:\n%s-----\n"), *MessageString);
+
+	const FTCHARToUTF8 Utf8MessageString(*MessageString);
+
+	const uint8* Data = reinterpret_cast<const uint8*>(Utf8MessageString.Get());
+	const int32 Length = Utf8MessageString.Length();
+	int32 TotalBytesSent = 0;
+
+	while (TotalBytesSent < Length)
+	{
+		const bool bReadyToWrite = Socket->Wait(ESocketWaitConditions::WaitForWrite, FTimespan::FromSeconds(1.0));
+		if (bReadyToWrite == false) 
+		{
+			UE_LOG(LogRtspMedia, Warning, TEXT("Write timeout when sending RTSP message: %s"), *MessageString);
+			return false;
+		}
+
+		int32 BytesSent = 0;
+		const bool bSuccess = Socket->Send(Data + TotalBytesSent, Length - TotalBytesSent, BytesSent);
+		if (!bSuccess || BytesSent <= 0)
+		{
+			UE_LOG(LogRtspMedia, Warning, TEXT("Failed to send RTSP message: %s"), *MessageString);
+			return false;
+		}
+
+		TotalBytesSent += BytesSent;
+	}
+
+	return true;
+}
+
+bool FRtspClient::SendRequest(const FRtspMessage& InRequest)
+{
+	const TOptional <int32> RequestCommandId = InRequest.GetCommandId();
+	if (!ensureMsgf(RequestCommandId.IsSet(), TEXT("Cannot send request if command ID is not set")))
+	{
+		return false;
+	}
+	
+	if (SendMessage(InRequest))
+	{
+		// Track the request so we can handle the response.
+		RequestTracker.AddPendingRequest(RequestCommandId.GetValue(), FRtspPendingRequest{InRequest, FPlatformTime::Seconds()});
+		return true;
+	}
+	
+	return false;
+}
+
+bool FRtspClient::ParseUrl(const FString& InUrl, FString& OutHost, int32& OutPort, FString& OutPath)
+{
+	// URL should match format: rtsp://host:port/path or rtsp://host/path
+
+	const FString Scheme = "rtsp://";
+	if (!InUrl.StartsWith(Scheme, ESearchCase::IgnoreCase))
+	{
+		return false;
+	}
+
+	FString Remainder = InUrl.Mid(Scheme.Len());
+	int32 FirstSlashIndex = Remainder.Find(TEXT("/"));
+	FString HostPort;
+
+	if (FirstSlashIndex != INDEX_NONE)
+	{
+		HostPort = Remainder.Left(FirstSlashIndex);
+		OutPath = Remainder.Mid(FirstSlashIndex + 1);
+	}
+	else
+	{
+		HostPort = Remainder;
+		OutPath = FString();
+	}
+
+	int32 ColonIndex = HostPort.Find(TEXT(":"));
+
+	int32 DefaultRtspPort = 554;
+	if (ColonIndex != INDEX_NONE)
+	{
+		OutHost = HostPort.Left(ColonIndex);
+		OutPort = FCString::Atoi(*HostPort.Mid(ColonIndex + 1));
+
+		if (OutPort <= 0 || OutPort > 65535)
+		{
+			OutPort = DefaultRtspPort;
+		}
+	}
+	else 
+	{
+		OutHost = HostPort;
+		OutPort = DefaultRtspPort;
+	}
+
+	return !OutHost.IsEmpty();
+}
+
+void FRtspClient::TeardownConnection()
+{
+	if (!SessionId.IsEmpty() && Socket.IsValid() && Socket->GetConnectionState() == SCS_Connected)
+	{
+		SendTeardown();
+	}
+
+	if (Socket.IsValid())
+	{
+		Socket->Close();
+		Socket.Reset();
+	}
+
+	SessionId.Empty();
+	CommandId = 0;
+	SdpSession.Reset();
+
+	JitterBuffer.Reset();
+	RequestTracker.Clear();
+
+	ReceiveBuffer.Reset();
+	ProcessBuffer.Reset();
+}
+
+int32 FRtspClient::HandleRunError()
+{
+	TeardownConnection();
+	SetState(ERtspClientState::Error);
+	return 1;
+}
diff --git a/Engine/Plugins/Media/RTSPMedia/Source/RTSPMedia/Private/RTSP/RtspClient.h b/Engine/Plugins/Media/RTSPMedia/Source/RTSPMedia/Private/RTSP/RtspClient.h
new file mode 100644
--- /dev/null
+++ b/Engine/Plugins/Media/RTSPMedia/Source/RTSPMedia/Private/RTSP/RtspClient.h
@@ -0,0 +1,150 @@
+// Copyright Epic Games, Inc. All Rights Reserved.
+
+#pragma once
+
+#include "RTP/RtpJitterBuffer.h"
+#include "RtspMediaConstants.h"
+#include "RtspMediaDefaults.h"
+#include "RtspMessage.h"
+#include "RtspRequestTracker.h"
+
+#include "CoreMinimal.h"
+#include "HAL/Runnable.h"
+#include "HAL/RunnableThread.h"
+#include "Sockets.h"
+
+class FSocket;
+class FRtpParser;
+
+struct RtpPacket;
+struct FSdpSession;
+
+struct FRtspClientConfiguration
+{
+	uint32 JitterBufferDepthMs = RtspMedia::Default::JitterBufferDepthMs;
+	int32 SocketBufferSizeBytes = RtspMedia::Default::SocketBufferSizeBytes;
+	float RequestTimeoutSeconds = RtspMedia::Default::RequestTimeoutSeconds;
+	FRtspTransportConfiguration::TransportProtocol TransportProtocol = FRtspTransportConfiguration::TransportProtocol::TCP;
+	FString Url;
+};
+
+enum class ERtspClientState : uint8
+{
+	Disconnected,
+	Connecting,
+	Connected,
+	Ready,
+	Described,
+	Setup,
+	Playing,
+	Error
+};
+
+// Executed on the game thread
+DECLARE_DELEGATE_TwoParams(FOnRtspClientStateChanged, ERtspClientState, ERtspClientState);
+DECLARE_DELEGATE_OneParam(FOnSdpSessionReceived, const FSdpSession&);
+
+// Executed on the client worker thread to minimise latency
+DECLARE_DELEGATE_OneParam(FOnRtpPacketReceived, FRtpPacket);
+
+/**
+ * Uses a dedicated worker thread for:
+ * - Socket input and output
+ * - State tracking
+ * - RTSP message handling
+ * - RTP packet handling
+ * 
+ * The public methods (Initialize, Disconnect) are expected be called on the game thread,
+ * and the state changed and SDP session received delegates are triggered on the game thread.
+ *
+ * The on RTP packet received delegate is executed on the worker thread to avoid a hop
+ * to the game thread before the packet is processed by the decoder (which has its own
+ * worker thread).
+ * 
+ */
+class FRtspClient : public FRunnable
+{
+public:
+	virtual ~FRtspClient() override;
+
+	bool Initialize(const FRtspClientConfiguration& InConfiguration);
+	void Disconnect();
+
+	FOnRtspClientStateChanged OnStateChanged;
+	FOnSdpSessionReceived OnSdpSessionReceived;
+
+	FOnRtpPacketReceived OnRtpPacketReceived;
+
+protected:
+	virtual bool Init() override;
+	virtual uint32 Run() override;
+	virtual void Stop() override;
+
+private:
+	static const TCHAR* StateToString(const ERtspClientState InState);
+
+	ERtspClientState GetState() const;
+	void SetState(ERtspClientState InNewState);
+
+	void ProcessState();
+	bool ProcessIncomingData();
+	void ProcessInterleavedPacket();
+	bool ProcessRtspMessage();
+
+	bool HandleResponse(const FRtspMessage& InMessage);
+	bool HandleRequest(const FRtspMessage& InRequest);
+
+	bool HandleOptionsResponse(const int32 InStatusCodeValue, const bool bInSuccess, const FRtspMessage& InOriginalRequest);
+	bool HandleDescribeResponse(const bool bInSuccess, const FRtspMessage& InResponse, const FRtspMessage& InOriginalRequest);
+	bool HandleSetupResponse(const bool bInSuccess, const FRtspMessage& InResponse);
+
+	bool InitializeJitterBuffer(const uint32 InClockRate);
+
+	bool SendOptions();
+	bool SendDescribe();
+	bool SendSetup(const FString& InControlUrl, const int32 InTrackIndex);
+	bool SendPlay();
+	bool SendTeardown();
+	bool SendKeepAlive();
+
+	bool SendMessage(const FRtspMessage& InMessage);
+	bool SendRequest(const FRtspMessage& InRequest);
+
+	bool ParseUrl(const FString& InUrl, FString& OutHost, int32& OutPort, FString& OutPath);
+
+	void TeardownConnection();
+	int32 HandleRunError();
+	
+	FRtspClientConfiguration Configuration;
+	
+	FString Host;
+	int32 Port = RtspMedia::Default::Port;
+	FString Path;
+	FString FullUrl; // rtsp://host:port/path
+
+	TUniquePtr<FSocket> Socket;
+	TArray<uint8> ReceiveBuffer;
+	TArray<uint8> ProcessBuffer;
+
+	ERtspClientState State = ERtspClientState::Disconnected;
+	mutable FCriticalSection StateLock;
+
+	// CSeq (Command Sequence) number incremented per request
+	int32 CommandId = 0;
+	
+	FString SessionId;
+	int32 SessionTimeoutSeconds = 60;
+
+	double KeepAliveInterval = 30.0;
+	double LastKeepAlive = 0.0;
+
+	FRtspRequestTracker RequestTracker;
+
+	TUniquePtr<FSdpSession> SdpSession;
+
+	TUniquePtr<FRunnableThread> Thread;
+	std::atomic<bool> bStopping{false};
+
+	// Network jitter handling
+	TUniquePtr<FRtpJitterBuffer> JitterBuffer;
+};
diff --git a/Engine/Plugins/Media/RTSPMedia/Source/RTSPMedia/Private/RTSP/RtspMessage.cpp b/Engine/Plugins/Media/RTSPMedia/Source/RTSPMedia/Private/RTSP/RtspMessage.cpp
new file mode 100644
--- /dev/null
+++ b/Engine/Plugins/Media/RTSPMedia/Source/RTSPMedia/Private/RTSP/RtspMessage.cpp
@@ -0,0 +1,448 @@
+// Copyright Epic Games, Inc. All Rights Reserved.
+
+#include "RTSP/RtspMessage.h"
+
+#include "RtspMediaConstants.h"
+
+#include "CoreMinimal.h"
+
+FRtspMessage FRtspMessage::Request(const ERtspRequestMethod InMethod, const FString& InUrl, const int32 InCommandId)
+{
+	FRtspMessage Request;
+	Request.SetMessageType(ERtspMessageType::Request);
+	Request.SetRequestMethod(InMethod);
+	Request.SetRequestUrl(InUrl);
+	Request.SetCommandId(InCommandId);
+	return Request;
+}
+
+FRtspMessage FRtspMessage::Response(const int32 InStatusCode, const FString& InReason, int32 InCommandId)
+{
+	FRtspMessage Response;
+	Response.SetMessageType(ERtspMessageType::Response);
+	Response.SetResponseStatusCode(InStatusCode);
+	Response.SetResponseReason(InReason);
+	Response.SetCommandId(InCommandId);
+	return Response;
+}
+
+ERtspMessageParseResult FRtspMessage::Parse(TArrayView<const uint8> InBuffer, FRtspMessage& OutMessage, int32& OutBytesConsumed)
+{
+	OutMessage = FRtspMessage();
+	OutBytesConsumed = 0;
+	
+	if (InBuffer.Num() < 4)
+	{
+		UE_LOG(LogRtspMedia, VeryVerbose, TEXT("RTSP message parsing attempted with a buffer containing less than 4 bytes"));
+		return ERtspMessageParseResult::Incomplete;
+	}
+	
+	// Find the end of the header section which should always end with '\r\n\r\n'
+	int32 HeadersEndIndex = INDEX_NONE;
+	constexpr int32 HeadersEndPatternLength = 4;
+	constexpr uint8 HeadersEndPattern[] = {'\r', '\n', '\r', '\n'};
+	const int32 LastSearchIndex = InBuffer.Num() - HeadersEndPatternLength;
+	for (int32 i = 0; i <= LastSearchIndex; ++i)
+	{
+		if (FMemory::Memcmp(InBuffer.GetData() + i, HeadersEndPattern, HeadersEndPatternLength) == 0)
+		{
+			HeadersEndIndex = i;
+			break;
+		}
+	}
+
+	if (HeadersEndIndex == INDEX_NONE)
+	{
+		UE_LOG(LogRtspMedia, VeryVerbose, TEXT("RTSP message parsing didn't find the header end pattern."));
+		return ERtspMessageParseResult::Incomplete;
+	}
+
+	const auto HeadersStringConversion = StringCast<TCHAR>(reinterpret_cast<const UTF8CHAR*>(InBuffer.GetData()), HeadersEndIndex);
+	const FString HeadersString(HeadersStringConversion.Length(), HeadersStringConversion.Get());
+
+	TArray<FString> HeaderLines;
+	HeadersString.ParseIntoArrayLines(HeaderLines);
+
+	if (HeaderLines.IsEmpty())
+	{
+		UE_LOG(LogRtspMedia, Error, TEXT("RTSP message header contains zero lines"));
+		return ERtspMessageParseResult::Error;
+	}
+
+	if (!ParseStartLine(HeaderLines[0], OutMessage))
+	{
+		UE_LOG(LogRtspMedia, Error, TEXT("RTSP message start line parsing failed"));
+		return ERtspMessageParseResult::Error;
+	}
+	
+	for (int32 i = 1; i < HeaderLines.Num(); ++i)
+	{
+		ParseHeaderLine(HeaderLines[i], OutMessage);
+	}
+
+	const int32 BodyStartIndex = HeadersEndIndex + HeadersEndPatternLength;
+	const int32 BodyLength = OutMessage.GetContentLength();
+
+	if (BodyLength < 0)
+	{
+		return ERtspMessageParseResult::Error;
+	}
+	
+	if (InBuffer.Num() - BodyStartIndex < BodyLength)
+	{
+		return ERtspMessageParseResult::Incomplete;
+	}
+
+	// Get body
+	if (BodyLength > 0)
+	{
+		const auto BodyStringConversion = StringCast<TCHAR>(reinterpret_cast<const UTF8CHAR*>(InBuffer.GetData() + BodyStartIndex), BodyLength);
+		OutMessage.Body = FString(BodyStringConversion.Length(), BodyStringConversion.Get());
+	}
+
+	OutBytesConsumed = BodyStartIndex + BodyLength;
+	return ERtspMessageParseResult::Complete;
+}
+
+ERtspMessageType FRtspMessage::GetMessageType() const
+{
+	return MessageType;
+}
+
+TOptional<int32> FRtspMessage::GetCommandId() const
+{
+	TOptional<FString> Value = GetHeaderValue(TEXT("CSeq"));
+	if (Value.IsSet())
+	{
+		const FString StringCommandId = Value.GetValue();
+		if (StringCommandId.IsNumeric())
+		{
+			return FCString::Atoi(*StringCommandId);
+		}
+	}
+	return {};
+}
+
+TOptional<ERtspRequestMethod> FRtspMessage::GetRequestMethod() const
+{
+	return RequestMethod;
+}
+
+TOptional<FString> FRtspMessage::GetRequestUrl() const
+{
+	return RequestUrl;
+}
+
+TOptional<int32> FRtspMessage::GetResponseStatusCode() const
+{
+	return ResponseStatusCode;
+}
+
+TOptional<FString> FRtspMessage::GetResponseReason() const
+{
+	return ResponseReason;
+}
+
+TOptional<FString> FRtspMessage::GetContentBaseUrl() const
+{
+	return GetHeaderValue(TEXT("Content-Base"));
+}
+
+int32 FRtspMessage::GetContentLength() const
+{
+	TOptional<FString> ContentLengthHeader = GetHeaderValue(TEXT("Content-Length"));
+	
+	if (ContentLengthHeader.IsSet())
+	{
+		const FString ContentLengthString = ContentLengthHeader.GetValue();
+		if (!ContentLengthString.IsNumeric())
+		{
+			UE_LOG(LogRtspMedia, Warning, TEXT("RTSP message content length header value must be numeric. Was: %s"), *ContentLengthString);
+			return -1;
+		}
+		const int32 ContentLength = FCString::Atoi(*ContentLengthString);
+		if (ContentLength >= 0)
+		{
+			return ContentLength;
+		}
+		UE_LOG(LogRtspMedia, Warning, TEXT("Invalid Content-Length value found within RTSP message header: %d"), ContentLength);
+		return -1;
+	}
+	
+	return 0;
+}
+
+TOptional<FString> FRtspMessage::GetHeaderValue(const FString& InName) const
+{
+	if (const FString* Value = Headers.Find(InName))
+	{
+		return *Value;
+	}
+	return {};
+}
+
+TOptional<FString> FRtspMessage::GetSessionId() const
+{
+	TOptional<FString> SessionHeader = GetHeaderValue(TEXT("Session"));
+	if (!SessionHeader.IsSet())
+	{
+		return {};
+	}
+
+	// Session header format is 'sessionid;timeout={number}'
+	const int32 SemicolonIndex = SessionHeader->Find(TEXT(";"));
+	FString SessionIdString;
+	if (SemicolonIndex != INDEX_NONE)
+	{
+		SessionIdString = SessionHeader->Left(SemicolonIndex).TrimStartAndEnd();
+	}
+	// If there is only one value in the header (i.e. no ';' delimiter)
+	// then return the complete header.
+	else
+	{
+		SessionIdString = SessionHeader.GetValue();
+	}
+	
+	return SessionIdString.TrimStartAndEnd();
+}
+
+TOptional<int32> FRtspMessage::GetSessionTimeout() const
+{
+	TOptional<FString> SessionHeader = GetHeaderValue(TEXT("Session"));
+	if (!SessionHeader.IsSet())
+	{
+		return {};
+	}
+
+	const FString TimeoutPrefix = TEXT("timeout=");
+	const int32 TimeoutIndex = SessionHeader->Find(TimeoutPrefix);
+	if (TimeoutIndex == INDEX_NONE)
+	{
+		return {};
+	}
+
+	const FString TimeoutInclusiveSuffix = SessionHeader->Mid(TimeoutIndex + TimeoutPrefix.Len());
+	const int32 PostTimeoutSemiColonIndex = TimeoutInclusiveSuffix.Find(TEXT(";"));
+
+	FString TimeoutString;
+	if (PostTimeoutSemiColonIndex != INDEX_NONE)
+	{
+		TimeoutString = TimeoutInclusiveSuffix.Left(PostTimeoutSemiColonIndex);
+	}
+	else
+	{
+		// If this is last value in the 'Session' header then return the remaining string. 
+		TimeoutString = TimeoutInclusiveSuffix;
+	}
+
+	if (TimeoutString.IsNumeric())
+	{
+		return FCString::Atoi(*TimeoutString);
+	}
+
+	return {};
+}
+
+const FString& FRtspMessage::GetBody() const
+{
+	return Body;
+}
+
+void FRtspMessage::SetMessageType(const ERtspMessageType InMessageType)
+{
+	MessageType = InMessageType;
+}
+
+void FRtspMessage::SetCommandId(const int32 InCommandId)
+{
+	SetHeader(TEXT("CSeq"), FString::FromInt(InCommandId));
+}
+
+void FRtspMessage::SetHeader(const FString& InName, const FString& InValue)
+{
+	Headers.Emplace(InName, InValue);
+}
+
+void FRtspMessage::SetSession(const FString& InSessionId)
+{
+	SetHeader(TEXT("Session"), InSessionId);
+}
+
+void FRtspMessage::SetTransport(const FRtspTransportConfiguration& InConfiguration)
+{
+	SetHeader(TEXT("Transport"), InConfiguration.BuildHeader());
+}
+
+void FRtspMessage::SetAccept(const FString& InMimeType)
+{
+	SetHeader(TEXT("Accept"), InMimeType);
+}
+
+void FRtspMessage::SetBody(const FString& InBody)
+{
+	Body = InBody;
+
+	// Ensure we base the Content-Length on actual bytes, not number of characters. 
+	const FTCHARToUTF8 Utf8BodyString(*InBody);
+	SetHeader(TEXT("Content-Length"), FString::FromInt(Utf8BodyString.Length()));
+}
+
+void FRtspMessage::SetRequestMethod(const ERtspRequestMethod InMethod)
+{
+	RequestMethod = InMethod;
+}
+
+void FRtspMessage::SetRequestUrl(const FString& InUrl)
+{
+	RequestUrl = InUrl;
+}
+
+void FRtspMessage::SetResponseStatusCode(const int32 InStatusCode)
+{
+	ResponseStatusCode = InStatusCode;
+}
+
+void FRtspMessage::SetResponseReason(const FString& InReason)
+{
+	ResponseReason = InReason;
+}
+
+FString FRtspMessage::ToString() const
+{
+	TArray<FString> Lines;
+
+	if (MessageType == ERtspMessageType::Request && RequestMethod.IsSet() && RequestUrl.IsSet())
+	{
+		// E.g. "OPTIONS rtsp://192.168.1.50/stream RTSP/1.0"
+		TArray<FString> RequestLineComponents;
+		RequestLineComponents.Add(RequestMethodToString(RequestMethod.GetValue()));
+		RequestLineComponents.Add(RequestUrl.GetValue());
+		RequestLineComponents.Add(TEXT("RTSP/1.0"));
+		Lines.Add(FString::Join(RequestLineComponents, TEXT(" ")));
+	}
+	else if (MessageType == ERtspMessageType::Response && ResponseStatusCode.IsSet() && ResponseReason.IsSet())
+	{
+		// E.g. "RTSP/1.0 200 OK"
+		TArray<FString> ResponseLineComponents;
+		ResponseLineComponents.Add(TEXT("RTSP/1.0"));
+		ResponseLineComponents.Add(FString::Printf(TEXT("%d"), ResponseStatusCode.GetValue()));
+		ResponseLineComponents.Add(ResponseReason.GetValue());
+		Lines.Add(FString::Join(ResponseLineComponents, TEXT(" ")));
+	}
+	else
+	{
+		return FString();
+	}
+
+	// Headers
+	for (const TPair<FString, FString>& Pair : Headers)
+	{
+		Lines.Add(FString::Printf(TEXT("%s: %s"), *Pair.Key, *Pair.Value));
+	}
+
+	// The RTSP standard requires carriage return-line feed between each line and a final CR/LF at the end of the header
+	return FString::Join(Lines, TEXT("\r\n")) + TEXT("\r\n\r\n") + Body;
+}
+
+FString FRtspMessage::RequestMethodToString(const ERtspRequestMethod InMethod)
+{
+	switch (InMethod)
+	{
+		case ERtspRequestMethod::Options: return TEXT("OPTIONS");
+		case ERtspRequestMethod::Describe: return TEXT("DESCRIBE");
+		case ERtspRequestMethod::Announce: return TEXT("ANNOUNCE");
+		case ERtspRequestMethod::Setup: return TEXT("SETUP");
+		case ERtspRequestMethod::Play: return TEXT("PLAY");
+		case ERtspRequestMethod::Pause: return TEXT("PAUSE");
+		case ERtspRequestMethod::Teardown: return TEXT("TEARDOWN");
+		case ERtspRequestMethod::GetParameter: return TEXT("GET_PARAMETER");
+		case ERtspRequestMethod::SetParameter: return TEXT("SET_PARAMETER");
+		case ERtspRequestMethod::Redirect: return TEXT("REDIRECT");
+		case ERtspRequestMethod::Record: return TEXT("RECORD");
+		case ERtspRequestMethod::Unknown:
+		default:
+		return TEXT("UNKNOWN");
+	}
+}
+
+ERtspRequestMethod FRtspMessage::StringToRequestMethod(const FString& InString)
+{
+	if (InString == TEXT("OPTIONS")) { return ERtspRequestMethod::Options; }
+	if (InString == TEXT("DESCRIBE")) { return ERtspRequestMethod::Describe; }
+	if (InString == TEXT("ANNOUNCE")) { return ERtspRequestMethod::Announce; }
+	if (InString == TEXT("SETUP")) { return ERtspRequestMethod::Setup; }
+	if (InString == TEXT("PLAY")) { return ERtspRequestMethod::Play; }
+	if (InString == TEXT("PAUSE")) { return ERtspRequestMethod::Pause; }
+	if (InString == TEXT("TEARDOWN")) { return ERtspRequestMethod::Teardown; }
+	if (InString == TEXT("GET_PARAMETER")) { return ERtspRequestMethod::GetParameter; }
+	if (InString == TEXT("SET_PARAMETER")) { return ERtspRequestMethod::SetParameter; }
+	if (InString == TEXT("REDIRECT")) { return ERtspRequestMethod::Redirect; }
+	if (InString == TEXT("RECORD")) { return ERtspRequestMethod::Record; }
+	return ERtspRequestMethod::Unknown;
+}
+
+bool FRtspMessage::ParseStartLine(const FString& InLine, FRtspMessage& OutMessage)
+{
+    // Responses start with RTSP/<version> and contain a status code and reason.
+    // E.g. "RTSP/1.0 404 Not Found"
+    // The reason string can contain spaces so we can't tokenise on " " for responses.
+    if (InLine.StartsWith(TEXT("RTSP/")))
+    {
+        OutMessage.MessageType = ERtspMessageType::Response;
+
+        const int32 FirstSpaceIndex = InLine.Find(TEXT(" "));
+        if (FirstSpaceIndex == INDEX_NONE)
+        {
+            UE_LOG(LogRtspMedia, Warning, TEXT("Malformed RTSP response start line, no space characters found: \"%s\""), *InLine);
+            return false;
+        }
+
+        const int32 SecondSpaceIndex = InLine.Find(TEXT(" "), ESearchCase::IgnoreCase, ESearchDir::FromStart, FirstSpaceIndex + 1);
+        if (SecondSpaceIndex == INDEX_NONE)
+        {
+            UE_LOG(LogRtspMedia, Warning, TEXT("Malformed RTSP response start line, second space character not found: \"%s\""), *InLine);
+            return false;
+        }
+
+        const FString StatusCodeString = InLine.Mid(FirstSpaceIndex + 1, SecondSpaceIndex - FirstSpaceIndex - 1);
+        OutMessage.ResponseStatusCode = FCString::Atoi(*StatusCodeString);
+        // The rest of the line contains the reason (e.g 'Not Found')
+        OutMessage.ResponseReason = InLine.Mid(SecondSpaceIndex + 1);
+    }
+    // Requests start with a METHOD, URL and RTSP/1.0
+    // E.g. "OPTIONS rtsp://<address>/<path> RTSP/1.0"
+    // The only spaces in a request status line are between the elements,
+    // so in this case we can tokenise on whitespace.
+    else
+    {
+        OutMessage.MessageType = ERtspMessageType::Request;
+
+        TArray<FString> RequestStatusLineComponents;
+        InLine.ParseIntoArray(RequestStatusLineComponents, TEXT(" "));
+
+        if (RequestStatusLineComponents.Num() != 3)
+        {
+            UE_LOG(LogRtspMedia, Warning, TEXT("Expected 3 components in RTSP request status line \"%s\" but found %d."), *InLine, RequestStatusLineComponents.Num());
+            return false;
+        }
+
+        OutMessage.RequestMethod = StringToRequestMethod(RequestStatusLineComponents[0]);
+    	OutMessage.RequestUrl = RequestStatusLineComponents[1];
+    }
+
+    return true;
+}
+
+void FRtspMessage::ParseHeaderLine(const FString& Line, FRtspMessage& OutMessage)
+{
+	// E.g. 'HeaderName: HeaderValue'
+	// Could be 'HeaderName:HeaderValue'
+	const int32 ColonIndex = Line.Find(TEXT(":"));
+	if (ColonIndex != INDEX_NONE)
+	{
+		// Account for non-compliant whitespace around the colon by trimming
+		const FString HeaderName = Line.Left(ColonIndex).TrimStartAndEnd();
+		const FString HeaderValue = Line.Mid(ColonIndex + 1).TrimStartAndEnd();
+		OutMessage.SetHeader(HeaderName, HeaderValue);
+	}
+}
diff --git a/Engine/Plugins/Media/RTSPMedia/Source/RTSPMedia/Private/RTSP/RtspMessage.h b/Engine/Plugins/Media/RTSPMedia/Source/RTSPMedia/Private/RTSP/RtspMessage.h
new file mode 100644
--- /dev/null
+++ b/Engine/Plugins/Media/RTSPMedia/Source/RTSPMedia/Private/RTSP/RtspMessage.h
@@ -0,0 +1,113 @@
+// Copyright Epic Games, Inc. All Rights Reserved.
+
+#pragma once
+
+#include "CoreMinimal.h"
+
+#include "RtspTransportConfiguration.h"
+
+/**
+ * We need to handle the possibility that incoming message header names are using the incorrect case (E.g. 'cseq' instead of 'CSeq').
+ * This custom KeyFuncs implementation will allow us to match names in a case-insensitive fashion.
+ */
+struct FCaseInsensitiveMapKeyMatch : TDefaultMapHashableKeyFuncs<FString, FString, false>
+{
+	static FORCEINLINE bool Matches(const FString& InA, const FString& InB)
+	{
+		return InA.Equals(InB, ESearchCase::IgnoreCase);
+	}
+
+	static FORCEINLINE uint32 GetKeyHash(const FString& InKey)
+	{
+		return GetTypeHash(InKey.ToLower());
+	}
+};
+
+enum class ERtspMessageType
+{
+	Unknown,
+	Request,
+	Response
+};
+
+enum class ERtspRequestMethod
+{
+	Unknown,
+	Options,
+	Describe,
+	Announce,
+	Setup,
+	Play,
+	Pause,
+	Teardown,
+	GetParameter,
+	SetParameter,
+	Redirect,
+	Record
+};
+
+enum class ERtspMessageParseResult : uint8
+{
+	Incomplete,
+	Complete,
+	Error
+};
+
+class FRtspMessage
+{
+public:
+	static FRtspMessage Request(const ERtspRequestMethod InMethod, const FString& InUrl, const int32 InCommandId);
+	static FRtspMessage Response(const int32 InStatusCode, const FString& InReason, const int32 InCommandId);
+	static ERtspMessageParseResult Parse(TArrayView<const uint8> InData, FRtspMessage& OutMessage, int32& OutBytesConsumed);
+	
+	static FString RequestMethodToString(const ERtspRequestMethod InMethod);
+	static ERtspRequestMethod StringToRequestMethod(const FString& InString);
+	
+	ERtspMessageType GetMessageType() const;
+
+	TOptional<int32> GetCommandId() const;
+	
+	TOptional<ERtspRequestMethod> GetRequestMethod() const;
+	TOptional<FString> GetRequestUrl() const;
+	TOptional<int32> GetResponseStatusCode() const;
+	TOptional<FString>  GetResponseReason() const;
+	TOptional<FString> GetContentBaseUrl() const;
+	int32 GetContentLength() const;
+	TOptional<FString> GetHeaderValue(const FString& InName) const;
+	TOptional<FString> GetSessionId() const;
+	TOptional<int32> GetSessionTimeout() const;
+	const FString& GetBody() const;
+
+	void SetMessageType(const ERtspMessageType InMessageType);
+	void SetCommandId(const int32 InCommandId);
+	void SetHeader(const FString& InName, const FString& InValue);
+	void SetSession(const FString& InSessionId);
+	void SetTransport(const FRtspTransportConfiguration& InConfiguration);
+	void SetAccept(const FString& InMimeType);
+	void SetBody(const FString& InBody);
+
+	void SetRequestMethod(const ERtspRequestMethod InMethod);
+	void SetRequestUrl(const FString& InUrl);
+	void SetResponseStatusCode(const int32 InStatusCode);
+	void SetResponseReason(const FString& InReason);
+
+	FString ToString() const;
+
+private:
+	static bool ParseStartLine(const FString& InLine, FRtspMessage& OutMessage);
+	static void ParseHeaderLine(const FString& InLine, FRtspMessage& OutMessage);
+	
+	ERtspMessageType MessageType = ERtspMessageType::Unknown;
+	
+	// Requests
+	TOptional<ERtspRequestMethod> RequestMethod;
+	TOptional<FString> RequestUrl;
+
+	// Responses
+	TOptional<int32> ResponseStatusCode;
+	TOptional<FString> ResponseReason;
+	
+	TMap<FString, FString, FDefaultSetAllocator, FCaseInsensitiveMapKeyMatch> Headers;
+	
+	FString Body;
+};
diff --git a/Engine/Plugins/Media/RTSPMedia/Source/RTSPMedia/Private/RTSP/RtspRequestTracker.cpp b/Engine/Plugins/Media/RTSPMedia/Source/RTSPMedia/Private/RTSP/RtspRequestTracker.cpp
new file mode 100644
--- /dev/null
+++ b/Engine/Plugins/Media/RTSPMedia/Source/RTSPMedia/Private/RTSP/RtspRequestTracker.cpp
@@ -0,0 +1,43 @@
+// Copyright Epic Games, Inc. All Rights Reserved.
+
+#include "RtspRequestTracker.h"
+
+#include "HAL/PlatformTime.h"
+
+void FRtspRequestTracker::AddPendingRequest(int32 InCommandId, const FRtspPendingRequest& InRequest)
+{
+	PendingRequests.Emplace(InCommandId, InRequest);
+}
+
+bool FRtspRequestTracker::HasPendingRequests() const
+{
+	return !PendingRequests.IsEmpty();
+}
+
+bool FRtspRequestTracker::TakePendingRequest(const int32 CommandId, FRtspPendingRequest& OutRequest)
+{
+	return PendingRequests.RemoveAndCopyValue(CommandId, OutRequest); 
+}
+
+TArray<FRtspPendingRequest> FRtspRequestTracker::GetExpiredRequests(const double TimeoutSeconds)
+{
+	const double Now = FPlatformTime::Seconds();
+	TArray<FRtspPendingRequest> ExpiredRequests;
+	auto Iterator = PendingRequests.CreateIterator();
+	while (Iterator)
+	{
+		FRtspPendingRequest& PendingRequest = Iterator.Value();
+		if (Now - PendingRequest.SentTimestamp > TimeoutSeconds)
+		{
+			ExpiredRequests.Emplace(MoveTemp(PendingRequest));
+			Iterator.RemoveCurrent();
+		}
+		++Iterator;
+	}
+	return ExpiredRequests;
+}
+
+void FRtspRequestTracker::Clear()
+{
+	PendingRequests.Empty();
+}
diff --git a/Engine/Plugins/Media/RTSPMedia/Source/RTSPMedia/Private/RTSP/RtspRequestTracker.h b/Engine/Plugins/Media/RTSPMedia/Source/RTSPMedia/Private/RTSP/RtspRequestTracker.h
new file mode 100644
--- /dev/null
+++ b/Engine/Plugins/Media/RTSPMedia/Source/RTSPMedia/Private/RTSP/RtspRequestTracker.h
@@ -0,0 +1,25 @@
+// Copyright Epic Games, Inc. All Rights Reserved.
+
+#pragma once
+
+#include "CoreMinimal.h"
+
+#include "RtspMessage.h"
+
+struct FRtspPendingRequest
+{
+	FRtspMessage Request;
+	double SentTimestamp = 0.0;
+};
+
+class FRtspRequestTracker
+{
+public:
+	void AddPendingRequest(int32 InCommandId, const FRtspPendingRequest& InRequest);
+	bool HasPendingRequests() const;
+	bool TakePendingRequest(const int32 CommandId, FRtspPendingRequest& OutRequest);
+	TArray<FRtspPendingRequest> GetExpiredRequests(const double TimeoutSeconds);
+	void Clear();
+private:
+	TMap<int32, FRtspPendingRequest> PendingRequests;
+};
diff --git a/Engine/Plugins/Media/RTSPMedia/Source/RTSPMedia/Private/RTSP/RtspTransportConfiguration.cpp b/Engine/Plugins/Media/RTSPMedia/Source/RTSPMedia/Private/RTSP/RtspTransportConfiguration.cpp
new file mode 100644
--- /dev/null
+++ b/Engine/Plugins/Media/RTSPMedia/Source/RTSPMedia/Private/RTSP/RtspTransportConfiguration.cpp
@@ -0,0 +1,50 @@
+// Copyright Epic Games, Inc. All Rights Reserved.
+
+#include "RTSP/RtspTransportConfiguration.h"
+
+FRtspTransportConfiguration FRtspTransportConfiguration::TcpInterleaved(int32 InRtpChannel, int32 InRtcpChannel)
+{
+	FRtspTransportConfiguration Configuration;
+	Configuration.Protocol = TransportProtocol::TCP;
+	Configuration.TcpInterleavedRtpChannel = InRtpChannel;
+	Configuration.TcpInterleavedRtcpChannel = InRtcpChannel;
+	return Configuration;
+}
+
+FRtspTransportConfiguration FRtspTransportConfiguration::Udp(int32 InClientRtpPort, int32 InClientRtcpPort)
+{
+	FRtspTransportConfiguration Configuration;
+	Configuration.Protocol = TransportProtocol::UDP;
+	Configuration.UdpClientRtpPort = InClientRtpPort;
+	Configuration.UdpClientRtcpPort = InClientRtcpPort;
+	return Configuration;
+}
+
+FString FRtspTransportConfiguration::BuildHeader() const
+{
+	TArray<FString> StackComponents;
+	StackComponents.Add(TEXT("RTP"));
+	StackComponents.Add(TEXT("AVP"));
+
+	FString ParameterKey;
+	FString ParameterValue;
+	
+	if (Protocol == TransportProtocol::TCP)
+	{
+		StackComponents.Add(TEXT("TCP"));
+		ParameterKey = TEXT("interleaved");
+		ParameterValue = FString::Printf(TEXT("%d-%d"), TcpInterleavedRtpChannel, TcpInterleavedRtcpChannel);
+	}
+	else if (Protocol == TransportProtocol::UDP)
+	{
+		ParameterKey = TEXT("client_port");
+		ParameterValue = FString::Printf(TEXT("%d-%d"), UdpClientRtpPort, UdpClientRtcpPort);
+	}
+
+	TArray<FString> Components;
+	Components.Add(FString::Join(StackComponents, TEXT("/")));
+	Components.Add(TEXT("unicast"));
+	Components.Add(ParameterKey + TEXT("=") + ParameterValue);
+
+	return FString::Join(Components, TEXT(";"));
+}
diff --git a/Engine/Plugins/Media/RTSPMedia/Source/RTSPMedia/Private/RTSP/RtspTransportConfiguration.h b/Engine/Plugins/Media/RTSPMedia/Source/RTSPMedia/Private/RTSP/RtspTransportConfiguration.h
new file mode 100644
--- /dev/null
+++ b/Engine/Plugins/Media/RTSPMedia/Source/RTSPMedia/Private/RTSP/RtspTransportConfiguration.h
@@ -0,0 +1,27 @@
+// Copyright Epic Games, Inc. All Rights Reserved.
+
+#pragma once
+
+#include "CoreMinimal.h"
+
+struct FRtspTransportConfiguration
+{
+	enum class TransportProtocol
+	{
+		TCP,
+		UDP
+	};
+	
+	TransportProtocol Protocol = TransportProtocol::TCP;
+
+	int32 TcpInterleavedRtpChannel = 0;
+	int32 TcpInterleavedRtcpChannel = 1;
+
+	int32 UdpClientRtpPort = 0;
+	int32 UdpClientRtcpPort = 0;
+
+	static FRtspTransportConfiguration TcpInterleaved(int32 InRtpChannel, int32 InRtcpChannel);
+	static FRtspTransportConfiguration Udp(int32 InClientRtpPort, int32 InClientRtcpPort);
+
+	FString BuildHeader() const;
+};
diff --git a/Engine/Plugins/Media/RTSPMedia/Source/RTSPMedia/Private/RtspMediaConstants.h b/Engine/Plugins/Media/RTSPMedia/Source/RTSPMedia/Private/RtspMediaConstants.h
new file mode 100644
--- /dev/null
+++ b/Engine/Plugins/Media/RTSPMedia/Source/RTSPMedia/Private/RtspMediaConstants.h
@@ -0,0 +1,32 @@
+// Copyright Epic Games, Inc. All Rights Reserved.
+
+#pragma once
+
+#include "Logging/LogMacros.h"
+#include "Misc/Guid.h"
+#include "UObject/NameTypes.h"
+
+DECLARE_LOG_CATEGORY_EXTERN(LogRtspMedia, Log, All);
+
+namespace RtspMedia
+{
+	inline const FName PlayerName(TEXT("RTSPMediaPlayer"));
+	inline constexpr FGuid PlayerPluginGUID(0x690B0293, 0x4C784EBF, 0xAF4769FF, 0x104000BD);
+
+	namespace Option
+	{
+		inline const FName AutoReconnect("AutoReconnect");
+		inline const FName DecoderBufferSize("DecoderBufferSize");
+		inline const FName DecoderPollIntervalMs("DecoderPollIntervalMs");
+		inline const FName JitterBufferDepthMs("JitterBufferDepthMs");
+		inline const FName MaxFragmentBufferSize("MaxFragmentBufferSize");
+		inline const FName MaxQueuedVideoSamples("MaxQueuedVideoSamples");
+		inline const FName MaxReconnectAttempts("MaxReconnectAttempts");
+		inline const FName MaxReconnectDelaySeconds("MaxReconnectDelaySeconds");
+		inline const FName MinReconnectDelaySeconds("MinReconnectDelaySeconds");
+		inline const FName RequestTimeoutSeconds("RequestTimeoutSeconds");
+		inline const FName SocketBufferSize("SocketBufferSize");
+		inline const FName TransportProtocol("TransportProtocol");
+	}
+	
+}
diff --git a/Engine/Plugins/Media/RTSPMedia/Source/RTSPMedia/Private/RtspMediaModule.cpp b/Engine/Plugins/Media/RTSPMedia/Source/RTSPMedia/Private/RtspMediaModule.cpp
new file mode 100644
--- /dev/null
+++ b/Engine/Plugins/Media/RTSPMedia/Source/RTSPMedia/Private/RtspMediaModule.cpp
@@ -0,0 +1,136 @@
+// Copyright Epic Games, Inc. All Rights Reserved.
+
+#include "Player/RtspMediaPlayer.h"
+#include "RtspMediaConstants.h"
+
+#include "IMediaModule.h"
+#include "IMediaPlayerFactory.h"
+#include "IMediaPlayer.h"
+#include "Internationalization/Internationalization.h"
+#include "Modules/ModuleInterface.h"
+#include "Modules/ModuleManager.h"
+
+DEFINE_LOG_CATEGORY(LogRtspMedia);
+
+#define LOCTEXT_NAMESPACE "RtspMediaModule"
+
+class FRtspMediaModule : public IModuleInterface, public IMediaPlayerFactory
+{
+public:
+	/** IModuleInterface implementation */
+	virtual void StartupModule() override;
+	virtual void ShutdownModule() override;
+
+	//~ Begin IMediaPlayerFactory
+	virtual bool CanPlayUrl(const FString& Url, const IMediaOptions* Options, TArray<FText>* OutWarnings, TArray<FText>* OutErrors) const override;
+	virtual int32 GetPlayabilityConfidenceScore(const FString& Url, const IMediaOptions* Options, TArray<FText>* OutWarnings, TArray<FText>* OutErrors) const override;
+	virtual TSharedPtr<IMediaPlayer, ESPMode::ThreadSafe> CreatePlayer(IMediaEventSink& InEventSink) override;
+	virtual FText GetDisplayName() const override;
+	virtual FName GetPlayerName() const override;
+	virtual FGuid GetPlayerPluginGUID() const override;
+	virtual const TArray<FString>& GetSupportedPlatforms() const override;
+	virtual bool SupportsFeature(EMediaFeature Feature) const override;
+	//~ End IMediaPlayerFactory
+
+private:
+	/** List of platforms that the media player support. */
+	TArray<FString> SupportedPlatforms;
+};
+
+void FRtspMediaModule::StartupModule()
+{
+	// supported platforms
+	SupportedPlatforms.Add(TEXT("Windows"));
+	SupportedPlatforms.Add(TEXT("Mac"));
+	SupportedPlatforms.Add(TEXT("Linux"));
+
+	// register player factory
+	if (IMediaModule* MediaModule = FModuleManager::LoadModulePtr<IMediaModule>("Media"))
+	{
+		MediaModule->RegisterPlayerFactory(*this);
+	}
+}
+
+void FRtspMediaModule::ShutdownModule()
+{
+	// unregister player factory
+	if (IMediaModule* MediaModule = FModuleManager::GetModulePtr<IMediaModule>("Media"))
+	{
+		MediaModule->UnregisterPlayerFactory(*this);
+	}
+}
+
+bool FRtspMediaModule::CanPlayUrl(const FString& Url, const IMediaOptions* Options, TArray<FText>* OutWarnings, TArray<FText>* OutErrors) const
+{
+	return GetPlayabilityConfidenceScore(Url, Options, OutWarnings, OutErrors) > 0;
+}
+
+int32 FRtspMediaModule::GetPlayabilityConfidenceScore(const FString& Url, const IMediaOptions* Options, TArray<FText>* OutWarnings, TArray<FText>* OutErrors) const
+{
+	FString Scheme;
+	FString Location;
+
+	// get scheme
+	if (!Url.Split(TEXT("://"), &Scheme, &Location, ESearchCase::CaseSensitive))
+	{
+		if (OutErrors != nullptr)
+		{
+			OutErrors->Add(LOCTEXT("NoSchemeFound", "No URI scheme found"));
+		}
+
+		return 0;
+	}
+
+	if (!Scheme.Equals(TEXT("rtsp"), ESearchCase::IgnoreCase))
+	{
+		if (OutErrors != nullptr)
+		{
+			OutErrors->Add(FText::Format(LOCTEXT("SchemeNotSupported", "The URI scheme '{0}' is not supported"), FText::FromString(Scheme)));
+		}
+
+		return 0;
+	}
+
+	// Returning a score less than 80 in order to avoid taking priority over the
+	// native Android and Windows (WMF) implementations (which return 80).
+	// Once we are confident that this module provides superior
+	// functionality we will consider increasing this score.
+	// This only impacts when a StreamMediaSource is used, for our
+	// dedicated RTSPMediaSource our own player implementation will
+	// always be used.
+	return 70;
+}
+
+TSharedPtr<IMediaPlayer, ESPMode::ThreadSafe> FRtspMediaModule::CreatePlayer(IMediaEventSink& InEventSink) 
+{
+	return MakeShared<FRtspMediaPlayer, ESPMode::ThreadSafe>(InEventSink);
+}
+
+FText FRtspMediaModule::GetDisplayName() const
+{
+	return LOCTEXT("DisplayName", "RTSP Media");
+}
+
+FName FRtspMediaModule::GetPlayerName() const
+{
+	return RtspMedia::PlayerName;
+}
+
+FGuid FRtspMediaModule::GetPlayerPluginGUID() const
+{
+	return RtspMedia::PlayerPluginGUID;
+}
+
+const TArray<FString>& FRtspMediaModule::GetSupportedPlatforms() const
+{
+	return SupportedPlatforms;
+}
+
+bool FRtspMediaModule::SupportsFeature(EMediaFeature Feature) const
+{
+	return Feature == EMediaFeature::VideoSamples;
+}
+
+#undef LOCTEXT_NAMESPACE
+
+IMPLEMENT_MODULE(FRtspMediaModule, RTSPMedia)
diff --git a/Engine/Plugins/Media/RTSPMedia/Source/RTSPMedia/Private/RtspMediaSource.cpp b/Engine/Plugins/Media/RTSPMedia/Source/RTSPMedia/Private/RtspMediaSource.cpp
new file mode 100644
--- /dev/null
+++ b/Engine/Plugins/Media/RTSPMedia/Source/RTSPMedia/Private/RtspMediaSource.cpp
@@ -0,0 +1,135 @@
+// Copyright Epic Games, Inc. All Rights Reserved.
+
+#include "RtspMediaSource.h"
+
+#include "RtspMediaConstants.h"
+
+FString URtspMediaSource::GetUrl() const
+{
+	const FString TrimmedHost = Host.TrimStartAndEnd();
+	const FString TrimmedPath = Path.TrimStartAndEnd();
+
+	FString Url = FString::Printf(TEXT("rtsp://%s:%d"), *TrimmedHost, Port);
+	if (!TrimmedPath.IsEmpty())
+	{
+		Url += "/" + TrimmedPath;
+	}
+	return Url;
+}
+
+bool URtspMediaSource::Validate() const
+{
+	const FString TrimmedHost = Host.TrimStartAndEnd();
+	const FString TrimmedPath = Path.TrimStartAndEnd();
+	
+	if (TrimmedHost.IsEmpty() ||
+		Port <= 0 ||
+		Port > TNumericLimits<uint16>::Max() ||
+		TrimmedPath.StartsWith(TEXT("/")))
+	{
+		return false;
+	}
+	
+	return true;
+}
+
+bool URtspMediaSource::GetMediaOption(const FName& InKey, bool InDefaultValue) const
+{
+	if (InKey == RtspMedia::Option::AutoReconnect)
+	{
+		return bAutoReconnect;
+	}
+	
+	return Super::GetMediaOption(InKey, InDefaultValue);
+}
+
+FName URtspMediaSource::GetDesiredPlayerName() const
+{
+	return RtspMedia::PlayerName;
+}
+
+int64 URtspMediaSource::GetMediaOption(const FName& InKey, int64 InDefaultValue) const
+{
+	if (InKey == RtspMedia::Option::DecoderBufferSize)
+	{
+		return DecoderBufferSize;
+	}
+
+	if (InKey == RtspMedia::Option::DecoderPollIntervalMs)
+	{
+		return DecoderPollIntervalMs;
+	}
+
+	if (InKey == RtspMedia::Option::JitterBufferDepthMs)
+	{
+		return JitterBufferDepthMs;
+	}
+	
+	if (InKey == RtspMedia::Option::MaxQueuedVideoSamples)
+	{
+		return MaxQueuedVideoSamples;
+	}
+
+	if (InKey == RtspMedia::Option::MaxFragmentBufferSize)
+	{
+		return MaxFragmentBufferSizeMb;
+	}
+
+	if (InKey == RtspMedia::Option::MaxReconnectAttempts)
+	{
+		return MaxReconnectAttempts;
+	}
+
+	if (InKey == RtspMedia::Option::TransportProtocol)
+	{
+		return static_cast<int64>(TransportProtocol);
+	}
+
+	if (InKey == RtspMedia::Option::SocketBufferSize)
+	{
+		return static_cast<int64>(SocketBufferSizeKb);
+	}
+
+	return Super::GetMediaOption(InKey, InDefaultValue);
+}
+
+double URtspMediaSource::GetMediaOption(const FName& InKey, double InDefaultValue) const
+{
+	if (InKey == RtspMedia::Option::MaxReconnectDelaySeconds)
+	{
+		return MaxReconnectDelaySeconds;
+	}
+	
+	if (InKey == RtspMedia::Option::MinReconnectDelaySeconds)
+	{
+		return MinReconnectDelaySeconds;
+	}
+	
+	if (InKey == RtspMedia::Option::RequestTimeoutSeconds)
+	{
+		return RequestTimeoutSeconds;
+	}
+	
+	return Super::GetMediaOption(InKey, InDefaultValue);
+}
+
+bool URtspMediaSource::HasMediaOption(const FName& Key) const
+{
+	if (Key == RtspMedia::Option::AutoReconnect ||
+		Key == RtspMedia::Option::DecoderBufferSize ||
+		Key == RtspMedia::Option::DecoderPollIntervalMs ||
+		Key == RtspMedia::Option::JitterBufferDepthMs ||
+		Key == RtspMedia::Option::MaxFragmentBufferSize ||
+		Key == RtspMedia::Option::MaxQueuedVideoSamples ||
+		Key == RtspMedia::Option::MaxReconnectAttempts ||
+		Key == RtspMedia::Option::MaxReconnectDelaySeconds ||
+		Key == RtspMedia::Option::MinReconnectDelaySeconds ||
+		Key == RtspMedia::Option::RequestTimeoutSeconds ||
+		Key == RtspMedia::Option::SocketBufferSize ||
+		Key == RtspMedia::Option::TransportProtocol)
+	{
+		return true;
+	}
+
+	return Super::HasMediaOption(Key);
+}
diff --git a/Engine/Plugins/Media/RTSPMedia/Source/RTSPMedia/Private/SDP/SdpParser.cpp b/Engine/Plugins/Media/RTSPMedia/Source/RTSPMedia/Private/SDP/SdpParser.cpp
new file mode 100644
--- /dev/null
+++ b/Engine/Plugins/Media/RTSPMedia/Source/RTSPMedia/Private/SDP/SdpParser.cpp
@@ -0,0 +1,336 @@
+// Copyright Epic Games, Inc. All Rights Reserved.
+
+#include "SdpParser.h"
+
+#include "RtspMediaConstants.h"
+#include "SDP/SdpSession.h"
+
+#include "Misc/Base64.h"
+
+TUniquePtr<FSdpSession> FSdpParser::Parse(const FString& InSdpBody, const FString& InContentBaseUrl)
+{
+	if (InSdpBody.IsEmpty())
+	{
+		UE_LOG(LogRtspMedia, Error, TEXT("Empty SDP body"));
+		return nullptr;
+	}
+
+	TUniquePtr<FSdpSession> Session = MakeUnique<FSdpSession>();
+	Session->ContentBaseUrl = InContentBaseUrl;
+
+	TArray<FString> Lines;
+	InSdpBody.ParseIntoArrayLines(Lines);
+
+	TArray<FSdpMediaTrack> MediaTracks;
+
+	for (int32 i = 0; i < Lines.Num(); ++i) 
+	{
+		const FString& Line = Lines[i];
+
+		// Skip blank lines
+		if (Line.IsEmpty())
+		{
+			continue;
+		}
+
+		// Skip malformed lines
+		// SDP lines are in the following format "Type=Value"
+		if (Line.Len() < 2 || Line[1] != TEXT('='))
+		{
+			continue;
+		}
+
+		const TCHAR Type = Line[0];
+		FString Value = Line.Mid(2);
+
+		switch (Type)
+		{
+			// Session Name
+			case TEXT('s'):
+				Session->SessionName = Value;
+				break;
+			// Session Info
+			case TEXT('i'):
+				Session->SessionInfo = Value;
+				break;
+			// Media line
+			case TEXT('m'):
+				{
+					FSdpMediaTrack MediaTrack;
+					if (!ParseMediaLine(Value, MediaTrack))
+					{
+						UE_LOG(LogRtspMedia, Warning, TEXT("SDP parser encountered invalid media line '%s'."), *Value);
+						continue;
+					}
+					MediaTracks.Add(MoveTemp(MediaTrack));
+					break;
+				}
+			case TEXT('a'):
+				if (!MediaTracks.IsEmpty())
+				{
+					if (!ParseAttributeLine(Value, MediaTracks.Last()))
+					{
+						continue;
+					}
+				}
+				else
+				{
+					UE_LOG(LogRtspMedia, Verbose, TEXT("SDP parser encountered attribute line '%s' but no media track is defined at this point."), *Value);
+				}
+				break;
+			default:
+				// Other types are ignored
+				break;
+		}
+	}
+
+	// Split media tracks into video and audio within the final session
+	for (FSdpMediaTrack& MediaTrack : MediaTracks)
+	{
+		UE_LOG(LogRtspMedia, Verbose, TEXT("Parsed SDP track: Type: %s Codec: %s Clock Rate: %u Control URL: %s"), *MediaTrack.MediaType, *MediaTrack.Codec, MediaTrack.ClockRate, *MediaTrack.ControlUrl);
+
+		if (MediaTrack.MediaType == TEXT("video"))
+		{
+			Session->VideoTracks.Add(MoveTemp(MediaTrack));
+		}
+		else if (MediaTrack.MediaType == TEXT("audio"))
+		{
+			Session->AudioTracks.Add(MoveTemp(MediaTrack));
+		}
+		// There are other types of tracks which we ignore
+		// - text (subtitles)
+		// - application (generic metadata)
+		// - message (irrelevant for RTSP)
+	}
+
+	UE_LOG(LogRtspMedia, Verbose, TEXT("Parsed SDP: %d video tracks, %d audio tracks"), Session->VideoTracks.Num(), Session->AudioTracks.Num());
+
+	return Session;
+}
+
+bool FSdpParser::ParseMediaLine(const FString& InMediaLine, FSdpMediaTrack& OutMediaTrack)
+{
+	// E.g. "video 0 RTP/AVP 96"
+
+	TArray<FString> Parts;
+	InMediaLine.ParseIntoArray(Parts, TEXT(" "));
+
+	if (Parts.Num() >= 4)
+	{
+		OutMediaTrack.MediaType = Parts[0];
+		OutMediaTrack.Port = FCString::Atoi(*Parts[1]);
+		OutMediaTrack.Protocol =  Parts[2];
+		OutMediaTrack.PayloadType = FCString::Atoi(*Parts[3]);
+		return true;
+	}
+	else
+	{
+		UE_LOG(LogRtspMedia, Warning, TEXT("Ignoring malformed media line in SDP: %s"), *InMediaLine);
+		return false;
+	}
+}
+
+bool FSdpParser::ParseAttributeLine(const FString& InAttributeLine, FSdpMediaTrack& InMediaTrack)
+{
+	// E.g. "rtpmap:96 H264/90000"
+	// E.g. "fmtp:96 packetization-mode=1; profile-level-id=420020; sprop-parameter-sets=J0IAIKtAWgUNNwICAgI=,KM48gA=="
+	// E.g. "control:trackID=0"
+	
+	int32 ColonIndex;
+	if (InAttributeLine.FindChar(TEXT(':'), ColonIndex))
+	{
+		const FString AttributeName = InAttributeLine.Left(ColonIndex).TrimStartAndEnd();
+		const FString AttributeData = InAttributeLine.Mid(ColonIndex + 1).TrimStartAndEnd();
+
+		// RTP Map
+		if (AttributeName == TEXT("rtpmap"))
+		{
+			return ParseRtpMap(AttributeData, InMediaTrack);
+		}
+		// Format Parameters
+		else if (AttributeName == TEXT("fmtp"))
+		{
+			return ParseFormatParameters(AttributeData, InMediaTrack);
+		}
+		else if (AttributeName == TEXT("control"))
+		{
+			InMediaTrack.ControlUrl = AttributeData;
+		}
+	}
+
+	return true;
+}
+
+bool FSdpParser::ParseRtpMap(const FString& InRtpMap, FSdpMediaTrack& InTrack)
+{
+	// E.g. "96 H264/90000"
+	// E.g. "94 MPEG-4-GENERIC/44100/2"
+
+	int32 SpaceIndex;
+	if (!InRtpMap.FindChar(TEXT(' '), SpaceIndex))
+	{
+		UE_LOG(LogRtspMedia, Warning, TEXT("SDP track does not contain format information"));
+		return false;
+	}
+
+	const FString Encoding = InRtpMap.Mid(SpaceIndex + 1);
+
+	TArray<FString> EncodingParts;
+	Encoding.ParseIntoArray(EncodingParts, TEXT("/"));
+
+	if (EncodingParts.Num() >= 1)
+	{
+		InTrack.Codec = EncodingParts[0];
+	}
+	
+	if (EncodingParts.Num() >= 2)
+	{
+		const FString ClockRateString = EncodingParts[1];
+		const int32 ParsedClockRate = FCString::Atoi(*ClockRateString);
+		if (ParsedClockRate <= 0)
+		{
+			UE_LOG(LogRtspMedia, Warning, TEXT("SDP parser found invalid clock rate: '%s'"), *ClockRateString);
+			return false;
+		}
+		InTrack.ClockRate = static_cast<uint32>(ParsedClockRate);
+	}
+	
+	if (EncodingParts.Num() >= 3)
+	{
+		InTrack.Channels = FCString::Atoi(*EncodingParts[2]);
+	}
+
+	return true;
+}
+
+bool FSdpParser::ParseFormatParameters(const FString& InFormatParameters, FSdpMediaTrack& InTrack)
+{
+	// e.g. "96 parameter1=value1;parameter2=value2"
+
+	int32 SpaceIndex;
+	if (!InFormatParameters.FindChar(TEXT(' '), SpaceIndex))
+	{
+		return false;
+	}
+
+	// Raw parameters
+	InTrack.FormatParameters = InFormatParameters.Mid(SpaceIndex + 1);
+
+	TArray<FString> Parameters;
+	InTrack.FormatParameters.ParseIntoArray(Parameters, TEXT(";"));
+
+	for (const FString& Parameter : Parameters)
+	{
+		FString ParameterKey;
+		FString ParameterValue;
+
+		if (Parameter.Split(TEXT("="), &ParameterKey, &ParameterValue))
+		{
+			ParameterKey.TrimStartAndEndInline();
+			ParameterValue.TrimStartAndEndInline();
+
+			// H.264 Stream Property Parameter Sets
+			if (ParameterKey == TEXT("sprop-parameter-sets"))
+			{
+				if (!ParseH264StreamPropertyParameterSets(ParameterValue, InTrack))
+				{
+					return false;
+				}
+			}
+			// H.264 packetization mode
+			// 0: One NAL per RTP packet
+			// 1: Non-Interleaved: Allows for:
+			//   a. Aggregation (STAP-A) combining multiple NALs in a single RTP packet.
+			//   b. Fragmentation (FU-A) across packets.
+			// 2: Interleaved: Out-of-order transmission
+			else if (ParameterKey == TEXT("packetization-mode"))
+			{
+				InTrack.PacketizationMode = FCString::Atoi(*ParameterValue);
+			}
+			// MPEG4-GENERIC payload format (includes AAC)
+			else if (ParameterKey == TEXT("config"))
+			{
+				if (HexDecode(ParameterValue, InTrack.AudioConfig))
+				{
+					UE_LOG(LogRtspMedia, Verbose, TEXT("Decoded AAC config: %d bytes"), InTrack.AudioConfig.Num());
+				}
+				else
+				{
+					UE_LOG(LogRtspMedia, Warning, TEXT("Failed to decode hex AAC configuration string: %s"), *ParameterValue);
+					return false;
+				}
+			}
+		}
+	}
+
+	return true;
+}
+
+bool FSdpParser::ParseH264StreamPropertyParameterSets(const FString& InH264StreamPropertyParameterSets, FSdpMediaTrack& InTrack)
+{
+	// e.g. "Z2QAH6zZQFAFuhAAAAMAEAAAAwPI8YMZYA==,aOvjyyLA"
+
+	TArray<FString> Parts;
+	InH264StreamPropertyParameterSets.ParseIntoArray(Parts, TEXT(","));
+
+	if (Parts.Num() >= 1)
+	{
+		const FString SequenceParameterSetBase64String = Parts[0];
+		if (Base64Decode(SequenceParameterSetBase64String, InTrack.SequenceParameterSet))
+		{
+			UE_LOG(LogRtspMedia, Verbose, TEXT("Parsed sequence parameter set: %d bytes"), InTrack.SequenceParameterSet.Num());
+		}
+		else
+		{
+			UE_LOG(LogRtspMedia, Warning, TEXT("Failed to decode sequence parameter set from base 64 encoded string: %s"), *SequenceParameterSetBase64String);
+			return false;
+		}
+	}
+
+	if (Parts.Num() >= 2)
+	{
+		const FString PictureParameterSetBase64String = Parts[1];
+		if (Base64Decode(PictureParameterSetBase64String, InTrack.PictureParameterSet))
+		{
+			UE_LOG(LogRtspMedia, Verbose, TEXT("Parsed picture parameter set: %d bytes"), InTrack.PictureParameterSet.Num());
+		}
+		else
+		{
+			UE_LOG(LogRtspMedia, Warning, TEXT("Failed to decode picture parameter set from base 64 encoded string: %s"), *PictureParameterSetBase64String);
+			return false;
+		}
+	}
+
+	return true;
+}
+
+bool FSdpParser::Base64Decode(const FString& InBase64EncodedString, TArray<uint8>& OutBuffer)
+{
+	const FString Trimmed = InBase64EncodedString.TrimStartAndEnd();
+	return FBase64::Decode(Trimmed, OutBuffer);
+}
+
+bool FSdpParser::HexDecode(const FString& InHexEncodedString, TArray<uint8>& OutBuffer)
+{
+	TArray<uint8> Decoded;
+	const FString Trimmed = InHexEncodedString.TrimStartAndEnd();
+
+	// Check we have an even number of characters for hex decoding
+	if (Trimmed.Len() % 2 != 0)
+	{
+		return false;
+	}
+	
+	for (int32 i = 0; i + 1 < Trimmed.Len(); i += 2)
+	{
+		FString HexString = Trimmed.Mid(i, 2);
+		if (!FChar::IsHexDigit(HexString[0]) || !FChar::IsHexDigit(HexString[1]))
+		{
+			return false;
+		}
+		uint8 ByteValue = static_cast<uint8>(FCString::Strtoi(*HexString, nullptr, 16));
+		Decoded.Add(ByteValue);
+	}
+	OutBuffer = Decoded;
+	return true;
+}
diff --git a/Engine/Plugins/Media/RTSPMedia/Source/RTSPMedia/Private/SDP/SdpParser.h b/Engine/Plugins/Media/RTSPMedia/Source/RTSPMedia/Private/SDP/SdpParser.h
new file mode 100644
--- /dev/null
+++ b/Engine/Plugins/Media/RTSPMedia/Source/RTSPMedia/Private/SDP/SdpParser.h
@@ -0,0 +1,24 @@
+// Copyright Epic Games, Inc. All Rights Reserved.
+
+#pragma once
+
+#include "CoreMinimal.h"
+
+struct FSdpMediaTrack;
+struct FSdpSession;
+
+class FSdpParser
+{
+public:
+	static TUniquePtr<FSdpSession> Parse(const FString& InSdpBody, const FString& InContentBaseUrl);
+
+private:
+	static bool ParseMediaLine(const FString& InMediaLine, FSdpMediaTrack& OutMediaTrack);
+	static bool ParseAttributeLine(const FString& InAttributeLine, FSdpMediaTrack& InMediaTrack);
+	static bool ParseRtpMap(const FString& InRtpMap, FSdpMediaTrack& InTrack);
+	static bool ParseFormatParameters(const FString& InFormatParameters, FSdpMediaTrack& InTrack);
+	static bool ParseH264StreamPropertyParameterSets(const FString& InH264StreamPropertyParameterSets, FSdpMediaTrack& InTrack);
+
+	static bool Base64Decode(const FString& InBase64EncodedString, TArray<uint8>& OutBuffer);
+	static bool HexDecode(const FString& InHexEncodedString, TArray<uint8>& OutBuffer);
+};
diff --git a/Engine/Plugins/Media/RTSPMedia/Source/RTSPMedia/Private/SDP/SdpSession.cpp b/Engine/Plugins/Media/RTSPMedia/Source/RTSPMedia/Private/SDP/SdpSession.cpp
new file mode 100644
--- /dev/null
+++ b/Engine/Plugins/Media/RTSPMedia/Source/RTSPMedia/Private/SDP/SdpSession.cpp
@@ -0,0 +1,84 @@
+// Copyright Epic Games, Inc. All Rights Reserved.
+
+#include "SdpSession.h"
+
+bool FSdpSession::HasVideo() const
+{
+	return !VideoTracks.IsEmpty();
+}
+
+bool FSdpSession::HasAudio() const
+{
+	return !AudioTracks.IsEmpty();
+}
+
+const FSdpMediaTrack* FSdpSession::GetFirstVideoTrack() const
+{
+	return HasVideo() ? &VideoTracks[0] : nullptr;
+}
+
+const FSdpMediaTrack* FSdpSession::GetFirstAudioTrack() const
+{
+	return HasAudio() ? &AudioTracks[0] : nullptr;
+}
+
+FString FSdpSession::ResolveControlUrl(const FString& InControlUrl) const
+{
+	// If the control URL is '*' this is an aggregate control URL which means we use the base url to control all streams
+	// Control URLs can be specified at a session or track level.
+	if (InControlUrl == TEXT("*"))
+	{
+		return ContentBaseUrl;
+	}
+
+	const FString SchemeAuthorityDelimiter = TEXT("://");
+	
+	// Any URL containing '://' should be preceded by a schema and is therefore an absolute URL
+	if (InControlUrl.Contains(SchemeAuthorityDelimiter))
+	{
+		return InControlUrl;
+	}
+
+	// If a control url starts with '/' it's an absolute path which should be appended
+	// to the root server URL e.g. rtsp://example.com:8554/InControlUrl.
+	// As opposed to being appended to the end of an existing path within the content base URL
+	if (InControlUrl.StartsWith(TEXT("/")))
+	{
+		// Find the server root URL from the content base URL
+		const int32 SchemeEndIndex = ContentBaseUrl.Find(SchemeAuthorityDelimiter);
+		if (SchemeEndIndex != INDEX_NONE)
+		{
+			const int32 PathStart = ContentBaseUrl.Find(
+				TEXT("/"),
+				ESearchCase::IgnoreCase,
+				ESearchDir::FromStart,
+				SchemeEndIndex + SchemeAuthorityDelimiter.Len());
+
+			FString Root;
+			if (PathStart != INDEX_NONE)
+			{
+				Root = ContentBaseUrl.Left(PathStart);
+			}
+			else // The content base url doesn't contain a path
+			{
+				Root = ContentBaseUrl;
+			}
+
+			return Root + InControlUrl;
+		}
+	}
+
+	// Otherwise the control "url" (which should be a simple path at this stage) is appended to the control base URL
+	FString BaseUrl;
+	// Strip the trailing '/' character from the base URL
+	if (ContentBaseUrl.EndsWith(TEXT("/")))
+	{
+		BaseUrl = ContentBaseUrl.Left(ContentBaseUrl.Len() - 1);
+	}
+	else
+	{
+		BaseUrl = ContentBaseUrl;
+	}
+
+	return BaseUrl + TEXT("/") + InControlUrl;
+}
diff --git a/Engine/Plugins/Media/RTSPMedia/Source/RTSPMedia/Private/SDP/SdpSession.h b/Engine/Plugins/Media/RTSPMedia/Source/RTSPMedia/Private/SDP/SdpSession.h
new file mode 100644
--- /dev/null
+++ b/Engine/Plugins/Media/RTSPMedia/Source/RTSPMedia/Private/SDP/SdpSession.h
@@ -0,0 +1,60 @@
+// Copyright Epic Games, Inc. All Rights Reserved.
+
+#pragma once
+
+#include "CoreMinimal.h"
+
+struct FSdpMediaTrack
+{
+	// Media Description
+	// "m=video 0 RTP/AVP 96"
+	FString MediaType;
+	int32 Port = 0;
+	FString Protocol;
+	int32 PayloadType = 0;
+
+	// Codec (Attribute)
+	// "a=rtpmap:96 H264/90000"
+	FString Codec;
+	uint32 ClockRate = 0;
+	int32 Channels = 0;
+
+	// Control URL (Attribute)
+	// "a=control:track1" (relative)
+	// "a=control:rtsp://host/path/track1" (absolute)
+	FString ControlUrl;
+
+	// Format Parameters (Attribute)
+	// "a=fmtp: <params>"
+	FString FormatParameters;
+
+	// H.264 Specific Parameters
+	TArray<uint8> SequenceParameterSet;
+	TArray<uint8> PictureParameterSet;
+	int32 PacketizationMode = 0;
+
+	// AAC Specific Parameters
+	TArray<uint8> AudioConfig;
+};
+
+struct FSdpSession
+{
+	// "s=<session_name>"
+	FString SessionName;
+
+	// "i=<session_info>"
+	FString SessionInfo;
+
+	// Provides the base URL of the stream content.
+	// A track may be located at <ContentBase>/trackID=0
+	FString ContentBaseUrl;
+
+	TArray<FSdpMediaTrack> VideoTracks;
+	TArray<FSdpMediaTrack> AudioTracks;
+
+	bool HasVideo() const;
+	bool HasAudio() const;
+	const FSdpMediaTrack* GetFirstVideoTrack() const;
+	const FSdpMediaTrack* GetFirstAudioTrack() const;
+	FString ResolveControlUrl(const FString& InControlUrl) const;
+};
diff --git a/Engine/Plugins/Media/RTSPMedia/Source/RTSPMedia/Private/Tests/SdpSessionTests.cpp b/Engine/Plugins/Media/RTSPMedia/Source/RTSPMedia/Private/Tests/SdpSessionTests.cpp
new file mode 100644
--- /dev/null
+++ b/Engine/Plugins/Media/RTSPMedia/Source/RTSPMedia/Private/Tests/SdpSessionTests.cpp
@@ -0,0 +1,199 @@
+// Copyright Epic Games, Inc. All Rights Reserved.
+
+#include "SDP/SdpSession.h"
+
+#include "Misc/AutomationTest.h"
+
+#if WITH_DEV_AUTOMATION_TESTS
+
+// Aggregate control
+
+IMPLEMENT_SIMPLE_AUTOMATION_TEST(
+	FResolveControlUrl_AggregateControl,
+	"Plugins.RTSPMedia.SdpSession.ResolveControlUrl.AggregateControl.ReturnsContentBase",
+	EAutomationTestFlags::EditorContext | EAutomationTestFlags::SmokeFilter)
+bool FResolveControlUrl_AggregateControl::RunTest(const FString& Parameters)
+{
+	FSdpSession Session;
+	Session.ContentBaseUrl = TEXT("rtsp://example.com:8554/stream/");
+	TestEqual(TEXT("Aggregate control returns Content-Base"),
+		Session.ResolveControlUrl(TEXT("*")),
+		TEXT("rtsp://example.com:8554/stream/"));
+	return true;
+}
+
+// Absolute URL
+
+IMPLEMENT_SIMPLE_AUTOMATION_TEST(
+	FResolveControlUrl_AbsoluteUrl_RtspScheme,
+	"Plugins.RTSPMedia.SdpSession.ResolveControlUrl.AbsoluteUrl.RtspScheme",
+	EAutomationTestFlags::EditorContext | EAutomationTestFlags::SmokeFilter)
+bool FResolveControlUrl_AbsoluteUrl_RtspScheme::RunTest(const FString& Parameters)
+{
+	FSdpSession Session;
+	Session.ContentBaseUrl = TEXT("rtsp://example.com:8554/stream/");
+	TestEqual(TEXT("Absolute URL returned as-is"),
+		Session.ResolveControlUrl(TEXT("rtsp://other-server.com:554/track1")),
+		TEXT("rtsp://other-server.com:554/track1"));
+	return true;
+}
+
+IMPLEMENT_SIMPLE_AUTOMATION_TEST(
+	FResolveControlUrl_AbsoluteUrl_HttpScheme,
+	"Plugins.RTSPMedia.SdpSession.ResolveControlUrl.AbsoluteUrl.HttpScheme",
+	EAutomationTestFlags::EditorContext | EAutomationTestFlags::SmokeFilter)
+bool FResolveControlUrl_AbsoluteUrl_HttpScheme::RunTest(const FString& Parameters)
+{
+	FSdpSession Session;
+	Session.ContentBaseUrl = TEXT("rtsp://example.com:8554/stream/");
+	TestEqual(TEXT("HTTP scheme returned as-is"),
+		Session.ResolveControlUrl(TEXT("http://example.com/track1")),
+		TEXT("http://example.com/track1"));
+	return true;
+}
+
+// Absolute path
+
+IMPLEMENT_SIMPLE_AUTOMATION_TEST(
+	FResolveControlUrl_AbsolutePath_WithPath,
+	"Plugins.RTSPMedia.SdpSession.ResolveControlUrl.AbsolutePath.ContentBaseWithPath",
+	EAutomationTestFlags::EditorContext | EAutomationTestFlags::SmokeFilter)
+bool FResolveControlUrl_AbsolutePath_WithPath::RunTest(const FString& Parameters)
+{
+	FSdpSession Session;
+	Session.ContentBaseUrl = TEXT("rtsp://example.com:8554/stream/");
+	TestEqual(TEXT("Absolute path resolves against server root"),
+		Session.ResolveControlUrl(TEXT("/track1")),
+		TEXT("rtsp://example.com:8554/track1"));
+	return true;
+}
+
+IMPLEMENT_SIMPLE_AUTOMATION_TEST(
+	FResolveControlUrl_AbsolutePath_NoTrailingSlash,
+	"Plugins.RTSPMedia.SdpSession.ResolveControlUrl.AbsolutePath.ContentBaseNoTrailingSlash",
+	EAutomationTestFlags::EditorContext | EAutomationTestFlags::SmokeFilter)
+bool FResolveControlUrl_AbsolutePath_NoTrailingSlash::RunTest(const FString& Parameters)
+{
+	FSdpSession Session;
+	Session.ContentBaseUrl = TEXT("rtsp://example.com:8554/stream");
+	TestEqual(TEXT("Absolute path resolves against server root without trailing slash"),
+		Session.ResolveControlUrl(TEXT("/track1")),
+		TEXT("rtsp://example.com:8554/track1"));
+	return true;
+}
+
+IMPLEMENT_SIMPLE_AUTOMATION_TEST(
+	FResolveControlUrl_AbsolutePath_NoPath,
+	"Plugins.RTSPMedia.SdpSession.ResolveControlUrl.AbsolutePath.ContentBaseNoPath",
+	EAutomationTestFlags::EditorContext | EAutomationTestFlags::SmokeFilter)
+bool FResolveControlUrl_AbsolutePath_NoPath::RunTest(const FString& Parameters)
+{
+	FSdpSession Session;
+	Session.ContentBaseUrl = TEXT("rtsp://example.com:8554");
+	TestEqual(TEXT("Absolute path appended to authority"),
+		Session.ResolveControlUrl(TEXT("/track1")),
+		TEXT("rtsp://example.com:8554/track1"));
+	return true;
+}
+
+IMPLEMENT_SIMPLE_AUTOMATION_TEST(
+	FResolveControlUrl_AbsolutePath_DeepPath,
+	"Plugins.RTSPMedia.SdpSession.ResolveControlUrl.AbsolutePath.DeepAbsolutePath",
+	EAutomationTestFlags::EditorContext | EAutomationTestFlags::SmokeFilter)
+bool FResolveControlUrl_AbsolutePath_DeepPath::RunTest(const FString& Parameters)
+{
+	FSdpSession Session;
+	Session.ContentBaseUrl = TEXT("rtsp://example.com:8554/live/stream/");
+	TestEqual(TEXT("Deep absolute path still resolves against root"),
+		Session.ResolveControlUrl(TEXT("/deep/path/track1")),
+		TEXT("rtsp://example.com:8554/deep/path/track1"));
+	return true;
+}
+
+IMPLEMENT_SIMPLE_AUTOMATION_TEST(
+	FResolveControlUrl_AbsolutePath_NoPort,
+	"Plugins.RTSPMedia.SdpSession.ResolveControlUrl.AbsolutePath.ContentBaseNoPort",
+	EAutomationTestFlags::EditorContext | EAutomationTestFlags::SmokeFilter)
+bool FResolveControlUrl_AbsolutePath_NoPort::RunTest(const FString& Parameters)
+{
+	FSdpSession Session;
+	Session.ContentBaseUrl = TEXT("rtsp://example.com/stream/");
+	TestEqual(TEXT("Absolute path resolves against portless server root"),
+		Session.ResolveControlUrl(TEXT("/track1")),
+		TEXT("rtsp://example.com/track1"));
+	return true;
+}
+
+// Relative path
+
+IMPLEMENT_SIMPLE_AUTOMATION_TEST(
+	FResolveControlUrl_RelativePath_TrailingSlash,
+	"Plugins.RTSPMedia.SdpSession.ResolveControlUrl.RelativePath.ContentBaseTrailingSlash",
+	EAutomationTestFlags::EditorContext | EAutomationTestFlags::SmokeFilter)
+bool FResolveControlUrl_RelativePath_TrailingSlash::RunTest(const FString& Parameters)
+{
+	FSdpSession Session;
+	Session.ContentBaseUrl = TEXT("rtsp://example.com:8554/stream/");
+	TestEqual(TEXT("Relative path appended to Content-Base"),
+		Session.ResolveControlUrl(TEXT("trackID=0")),
+		TEXT("rtsp://example.com:8554/stream/trackID=0"));
+	return true;
+}
+
+IMPLEMENT_SIMPLE_AUTOMATION_TEST(
+	FResolveControlUrl_RelativePath_NoTrailingSlash,
+	"Plugins.RTSPMedia.SdpSession.ResolveControlUrl.RelativePath.ContentBaseNoTrailingSlash",
+	EAutomationTestFlags::EditorContext | EAutomationTestFlags::SmokeFilter)
+bool FResolveControlUrl_RelativePath_NoTrailingSlash::RunTest(const FString& Parameters)
+{
+	FSdpSession Session;
+	Session.ContentBaseUrl = TEXT("rtsp://example.com:8554/stream");
+	TestEqual(TEXT("Relative path appended with separator"),
+		Session.ResolveControlUrl(TEXT("trackID=0")),
+		TEXT("rtsp://example.com:8554/stream/trackID=0"));
+	return true;
+}
+
+IMPLEMENT_SIMPLE_AUTOMATION_TEST(
+	FResolveControlUrl_RelativePath_SimpleTrackName,
+	"Plugins.RTSPMedia.SdpSession.ResolveControlUrl.RelativePath.SimpleTrackName",
+	EAutomationTestFlags::EditorContext | EAutomationTestFlags::SmokeFilter)
+bool FResolveControlUrl_RelativePath_SimpleTrackName::RunTest(const FString& Parameters)
+{
+	FSdpSession Session;
+	Session.ContentBaseUrl = TEXT("rtsp://example.com:8554/live/");
+	TestEqual(TEXT("Simple track name appended"),
+		Session.ResolveControlUrl(TEXT("video")),
+		TEXT("rtsp://example.com:8554/live/video"));
+	return true;
+}
+
+IMPLEMENT_SIMPLE_AUTOMATION_TEST(
+	FResolveControlUrl_RelativePath_NoPath,
+	"Plugins.RTSPMedia.SdpSession.ResolveControlUrl.RelativePath.ContentBaseNoPath",
+	EAutomationTestFlags::EditorContext | EAutomationTestFlags::SmokeFilter)
+bool FResolveControlUrl_RelativePath_NoPath::RunTest(const FString& Parameters)
+{
+	FSdpSession Session;
+	Session.ContentBaseUrl = TEXT("rtsp://example.com:8554");
+	TestEqual(TEXT("Relative path appended to authority-only Content-Base"),
+		Session.ResolveControlUrl(TEXT("trackID=0")),
+		TEXT("rtsp://example.com:8554/trackID=0"));
+	return true;
+}
+
+IMPLEMENT_SIMPLE_AUTOMATION_TEST(
+	FResolveControlUrl_RelativePath_NoPort,
+	"Plugins.RTSPMedia.SdpSession.ResolveControlUrl.RelativePath.ContentBaseNoPort",
+	EAutomationTestFlags::EditorContext | EAutomationTestFlags::SmokeFilter)
+bool FResolveControlUrl_RelativePath_NoPort::RunTest(const FString& Parameters)
+{
+	FSdpSession Session;
+	Session.ContentBaseUrl = TEXT("rtsp://example.com/stream/");
+	TestEqual(TEXT("Relative path appended to portless Content-Base"),
+		Session.ResolveControlUrl(TEXT("trackID=0")),
+		TEXT("rtsp://example.com/stream/trackID=0"));
+	return true;
+}
+
+#endif // WITH_DEV_AUTOMATION_TESTS
diff --git a/Engine/Plugins/Media/RTSPMedia/Source/RTSPMedia/Public/RtspMediaDefaults.h b/Engine/Plugins/Media/RTSPMedia/Source/RTSPMedia/Public/RtspMediaDefaults.h
new file mode 100644
--- /dev/null
+++ b/Engine/Plugins/Media/RTSPMedia/Source/RTSPMedia/Public/RtspMediaDefaults.h
@@ -0,0 +1,24 @@
+// Copyright Epic Games, Inc. All Rights Reserved.
+
+#pragma once
+
+#include "CoreMinimal.h"
+
+namespace RtspMedia::Default
+{
+	inline constexpr bool bAutoReconnect = true;
+	inline constexpr int32 DecoderBufferSize = 5;
+	inline constexpr int32 DecoderPollIntervalMs = 5;
+	inline constexpr uint32 JitterBufferDepthMs = 1000;
+	inline constexpr int32 MaxFragmentBufferSizeMb = 8;
+	inline constexpr int64 MaxFragmentBufferSizeBytes = MaxFragmentBufferSizeMb * 1024 * 1024;
+	inline constexpr int32 MaxQueuedVideoSamples = 4;
+	inline constexpr int32 MaxQueuedAudioSamples = 4;
+	inline constexpr int32 MaxReconnectAttempts = 0;
+	inline constexpr float MaxReconnectDelaySeconds = 30.0; 
+	inline constexpr float MinReconnectDelaySeconds = 1.0;
+	inline constexpr int32 Port = 8554;
+	inline constexpr float RequestTimeoutSeconds = 5.0;
+	inline constexpr int32 SocketBufferSizeKb = 512;
+	inline constexpr int32 SocketBufferSizeBytes = SocketBufferSizeKb * 1024;
+}
diff --git a/Engine/Plugins/Media/RTSPMedia/Source/RTSPMedia/Public/RtspMediaSource.h b/Engine/Plugins/Media/RTSPMedia/Source/RTSPMedia/Public/RtspMediaSource.h
new file mode 100644
--- /dev/null
+++ b/Engine/Plugins/Media/RTSPMedia/Source/RTSPMedia/Public/RtspMediaSource.h
@@ -0,0 +1,118 @@
+// Copyright Epic Games, Inc. All Rights Reserved.
+
+#pragma once
+
+#include "BaseMediaSource.h"
+
+#include "RtspMediaDefaults.h"
+
+#include "RtspMediaSource.generated.h"
+
+UENUM(BlueprintType)
+enum class ERtspMediaTransportProtocol : uint8
+{
+	/** RTP packets are interleaved with RTSP control messages */
+	TCP
+};
+
+namespace RtspMedia::Default
+{
+	inline constexpr ERtspMediaTransportProtocol TransportProtocol = ERtspMediaTransportProtocol::TCP;
+}
+
+UCLASS(BlueprintType, hideCategories=(Platforms, Object), meta=(DisplayName="RTSP Media Source"))
+class RTSPMEDIA_API URtspMediaSource : public UBaseMediaSource
+{
+	GENERATED_BODY()
+
+public:
+
+	/** Host name or IP address of the RTSP server */
+	UPROPERTY(EditAnywhere, BlueprintReadWrite, Category = "Network")
+	FString Host;
+
+	/** Port number of the RTSP Server */
+	UPROPERTY(EditAnywhere, BlueprintReadWrite, Category = "Network", meta=(ClampMin="1", ClampMax="65535"))
+	int32 Port = RtspMedia::Default::Port;
+
+	/** Path to the stream on RTSP server */
+	UPROPERTY(EditAnywhere, BlueprintReadWrite, Category = "Network")
+	FString Path;
+
+	/** Only TCP interleaved transport is supported */
+	UPROPERTY(BlueprintReadOnly, Category="Network", meta=(Hidden))
+	ERtspMediaTransportProtocol TransportProtocol = RtspMedia::Default::TransportProtocol;
+	
+	/** Socket buffer size for stream data in KB.
+	 * Increase this setting for high bitrate streams (e.g. 4K) */
+	UPROPERTY(EditAnywhere, BlueprintReadWrite, Category="Network", AdvancedDisplay, meta=(DisplayName="Socket Buffer Size (KB)", ClampMin="128", ClampMax="4096"))
+	int32 SocketBufferSizeKb = RtspMedia::Default::SocketBufferSizeKb;
+
+	/**
+	 * The jitter buffer accounts for spikes in the network delivery of RTP packets in order to provide a steady stream of packets to the decoder.
+	 * A higher buffer depth will absorb more network jitter at the cost of additional latency.
+	 * On clean networks a low or zero value may be configured, and will minimise the latency introduced by the buffer. */
+	UPROPERTY(EditAnywhere, BlueprintReadWrite, Category="Network", AdvancedDisplay, meta=(DisplayName="Jitter Buffer Depth (ms)", ClampMin="0", ClampMax="5000"))
+	int32 JitterBufferDepthMs = RtspMedia::Default::JitterBufferDepthMs;
+	
+	/** RTSP request timeout in seconds.
+	 * Increase for slow or unreliable networks to avoid premature session failures */
+	UPROPERTY(EditAnywhere, BlueprintReadWrite, Category="Network", AdvancedDisplay, meta=(DisplayName="Request Timeout (s)", ClampMin="1.0", ClampMax="30.0"))
+	float RequestTimeoutSeconds = RtspMedia::Default::RequestTimeoutSeconds;
+	
+	/** Automatically reconnect when the connection fails
+	 * The player will begin reconnection attempts after the delay specified in Min Retry Delay */
+	UPROPERTY(EditAnywhere, BlueprintReadWrite, Category="Network", AdvancedDisplay)
+	bool bAutoReconnect = RtspMedia::Default::bAutoReconnect;
+
+	/** The minimum amount of time to wait before attempting to reconnect after a connection failure */
+	UPROPERTY(EditAnywhere, BlueprintReadWrite, Category="Network", AdvancedDisplay, meta=(DisplayName="Min Reconnect Delay (s)", EditCondition="bAutoReconnect", ClampMin="1.0", ClampMax="60.0"))
+	float MinReconnectDelaySeconds = RtspMedia::Default::MinReconnectDelaySeconds;
+
+	/** The maximum amount of time to wait before attempting to reconnect after a connection failure
+	 * The player will gradually increase the delay before each connection attempt but for no longer than this maximum */
+	UPROPERTY(EditAnywhere, BlueprintReadWrite, Category="Network", AdvancedDisplay, meta=(DisplayName="Max Reconnect Delay (s)", EditCondition="bAutoReconnect", ClampMin="1.0", ClampMax="300.0"))
+	float MaxReconnectDelaySeconds = RtspMedia::Default::MaxReconnectDelaySeconds;
+
+	/** The maximum number of times the player will attempt to reconnect after a connection failure
+	 * Set to 0 to retry indefinitely */
+	UPROPERTY(EditAnywhere, BlueprintReadWrite, Category="Network", AdvancedDisplay, meta=(EditCondition="bAutoReconnect", ClampMin="0", ClampMax="1000"))
+	int32 MaxReconnectAttempts = RtspMedia::Default::MaxReconnectAttempts;
+
+	/** Decoder reorder buffer size. Use 1 for minimum latency (i.e. no B-frames)
+	 * Increase if video frames appears out of order */
+	UPROPERTY(EditAnywhere, BlueprintReadWrite, Category = "Video", meta=(DisplayName="Decoder Buffer Size (frames)", ClampMin="1", ClampMax="16"))
+	int32 DecoderBufferSize = RtspMedia::Default::DecoderBufferSize;
+
+	/** How frequently to poll the decoder for output frames (in milliseconds)
+	 * Lower values reduce latency but increase CPU usage */
+	UPROPERTY(EditAnywhere, BlueprintReadWrite, Category = "Video", meta=(DisplayName="Decoder Poll Interval (ms)", ClampMin="1", ClampMax="10"))
+	int32 DecoderPollIntervalMs = RtspMedia::Default::DecoderPollIntervalMs;
+
+	/** Places a limit on how much memory is used for fragmented NAL unit reassembly
+	 * NAL units will be dropped if this memory limit is exceeded
+	 * Consider increasing this value for high bitrate streams */ 
+	UPROPERTY(EditAnywhere, BlueprintReadWrite, Category = "Video", meta=(DisplayName="Max Fragment Buffer Size (MB)", ClampMin="2", ClampMax="64"))
+	int32 MaxFragmentBufferSizeMb = RtspMedia::Default::MaxFragmentBufferSizeMb;
+
+	/** Maximum decoded video frames to queue while waiting for game thread consumption
+	 *  Higher values reduce visible skips after game thread hitches at the cost of memory */
+	UPROPERTY(EditAnywhere, BlueprintReadWrite, Category="Video", meta=(ClampMin="2", ClampMax="16"))
+	int32 MaxQueuedVideoSamples = RtspMedia::Default::MaxQueuedVideoSamples;
+
+	//~ Begin UMediaSource
+	virtual FString GetUrl() const override;
+	virtual bool Validate() const override;
+	//~ End UMediaSource
+
+	//~ Begin IMediaOptions
+	virtual bool GetMediaOption(const FName& InKey, bool InDefaultValue) const override;
+	virtual int64 GetMediaOption(const FName& InKey, int64 InDefaultValue) const override;
+	virtual double GetMediaOption(const FName& InKey, double InDefaultValue) const override;
+	virtual bool HasMediaOption(const FName& InKey) const override;
+	//~ End IMediaOptions
+
+	//~ Begin UBaseMediaSource
+	virtual FName GetDesiredPlayerName() const override;
+	//~ End UBaseMediaSource
+};
\ No newline at end of file
diff --git a/Engine/Plugins/Media/RTSPMedia/Source/RTSPMedia/RTSPMedia.Build.cs b/Engine/Plugins/Media/RTSPMedia/Source/RTSPMedia/RTSPMedia.Build.cs
new file mode 100644
--- /dev/null
+++ b/Engine/Plugins/Media/RTSPMedia/Source/RTSPMedia/RTSPMedia.Build.cs
@@ -0,0 +1,40 @@
+// Copyright Epic Games, Inc. All Rights Reserved.
+
+using System.IO;
+using UnrealBuildTool;
+
+public class RTSPMedia : ModuleRules
+{
+	public RTSPMedia(ReadOnlyTargetRules Target) : base(Target)
+	{
+		PCHUsage = PCHUsageMode.NoPCHs;
+		bRequiresPlatformSDK = true;
+		
+		PublicDependencyModuleNames.AddRange(new string[]
+		{
+			"Core",
+			"MediaAssets"
+		});
+
+		PrivateDependencyModuleNames.AddRange(new string[]
+		{
+			"CoreUObject",
+			"ElectraBase",
+			"ElectraDecoders",
+			"ElectraCodecFactory",
+			"ElectraSamples",
+			"Engine",
+			"Media",
+			"MediaIOCore",
+			"MediaUtils",
+			"Networking",
+			"RenderCore",
+			"RHI",
+			"Slate",
+			"SlateCore",
+			"Sockets",
+			"TimeManagement"
+		});
+
+	}
+}
\ No newline at end of file
diff --git a/Engine/Plugins/Media/RTSPMedia/Source/RTSPMediaEditor/Private/RtspMediaEditorModule.cpp b/Engine/Plugins/Media/RTSPMedia/Source/RTSPMediaEditor/Private/RtspMediaEditorModule.cpp
new file mode 100644
--- /dev/null
+++ b/Engine/Plugins/Media/RTSPMedia/Source/RTSPMediaEditor/Private/RtspMediaEditorModule.cpp
@@ -0,0 +1,106 @@
+// Copyright Epic Games, Inc. All Rights Reserved.
+
+#include "AssetToolsModule.h"
+#include "Interfaces/IPluginManager.h"
+#include "Modules/ModuleInterface.h"
+#include "Modules/ModuleManager.h"
+#include "RtspMediaSourceActions.h"
+#include "Styling/AppStyle.h"
+#include "Styling/SlateStyle.h"
+#include "Styling/SlateStyleRegistry.h"
+
+DEFINE_LOG_CATEGORY_STATIC(LogRtspMediaEditor, Log, All);
+
+#define LOCTEXT_NAMESPACE "RtspMediaEditorModule"
+
+class FRtspMediaEditorModule : public IModuleInterface
+{
+public:
+	/** IModuleInterface implementation */
+	virtual void StartupModule() override;
+	virtual void ShutdownModule() override;
+
+private:
+	void RegisterAssetTools();
+	void UnregisterAssetTools();
+
+	void RegisterStyle();
+	void UnregisterStyle();
+
+	TArray<TSharedRef<IAssetTypeActions>> RegisteredAssetTypeActions;
+	TUniquePtr<FSlateStyleSet> StyleInstance;
+};
+
+void FRtspMediaEditorModule::StartupModule()
+{
+	RegisterAssetTools();
+	RegisterStyle();
+}
+
+void FRtspMediaEditorModule::ShutdownModule()
+{
+	UnregisterAssetTools();
+	UnregisterStyle();
+}
+
+void FRtspMediaEditorModule::RegisterAssetTools()
+{
+	IAssetTools& AssetTools = FModuleManager::LoadModuleChecked<FAssetToolsModule>("AssetTools").Get();
+	TSharedRef<IAssetTypeActions> Action = MakeShared<FRtspMediaSourceActions>();
+	AssetTools.RegisterAssetTypeActions(Action);
+	RegisteredAssetTypeActions.Add(Action);
+}
+
+void FRtspMediaEditorModule::UnregisterAssetTools()
+{
+	if (FAssetToolsModule* AssetToolsModule = FModuleManager::GetModulePtr<FAssetToolsModule>("AssetTools"))
+	{
+		IAssetTools& AssetTools = AssetToolsModule->Get();
+		for (const TSharedRef<IAssetTypeActions>& Action : RegisteredAssetTypeActions)
+		{
+			AssetTools.UnregisterAssetTypeActions(Action);
+		}
+	}
+
+	RegisteredAssetTypeActions.Reset();
+}
+
+void FRtspMediaEditorModule::RegisterStyle()
+{
+	const TSharedPtr<IPlugin> Plugin = IPluginManager::Get().FindPlugin(TEXT("RTSPMedia"));
+
+	if (!Plugin.IsValid())
+	{
+		return;
+	}
+
+	const FString ContentDirectory = Plugin->GetContentDir();
+
+	StyleInstance = MakeUnique<FSlateStyleSet>("RtspMediaStyle");
+	StyleInstance->SetContentRoot(ContentDirectory / TEXT("Editor/Icons"));
+
+	StyleInstance->Set(
+		"ClassIcon.RtspMediaSource",
+		new FSlateVectorImageBrush(StyleInstance->RootToContentDir(TEXT("RtspMediaSource_16"), TEXT(".svg")), CoreStyleConstants::Icon16x16)
+	);
+
+	StyleInstance->Set(
+		"ClassThumbnail.RtspMediaSource", 
+		new FSlateVectorImageBrush(StyleInstance->RootToContentDir(TEXT("RtspMediaSource_64"), TEXT(".svg")), CoreStyleConstants::Icon64x64)
+	);
+
+	FSlateStyleRegistry::RegisterSlateStyle(*StyleInstance);
+}
+
+void FRtspMediaEditorModule::UnregisterStyle()
+{
+	if (StyleInstance.IsValid())
+	{
+		FSlateStyleRegistry::UnRegisterSlateStyle(*StyleInstance);
+		StyleInstance.Reset();
+	}
+}
+
+#undef LOCTEXT_NAMESPACE
+
+IMPLEMENT_MODULE(FRtspMediaEditorModule, RTSPMediaEditor)
diff --git a/Engine/Plugins/Media/RTSPMedia/Source/RTSPMediaEditor/Private/RtspMediaSourceActions.cpp b/Engine/Plugins/Media/RTSPMedia/Source/RTSPMediaEditor/Private/RtspMediaSourceActions.cpp
new file mode 100644
--- /dev/null
+++ b/Engine/Plugins/Media/RTSPMedia/Source/RTSPMediaEditor/Private/RtspMediaSourceActions.cpp
@@ -0,0 +1,29 @@
+// Copyright Epic Games, Inc. All Rights Reserved.
+
+#include "RtspMediaSourceActions.h"
+
+#include "RtspMediaSource.h"
+
+#define LOCTEXT_NAMESPACE "RtspMediaSourceActions"
+
+bool FRtspMediaSourceActions::CanFilter()
+{
+	return true;
+}
+
+FText FRtspMediaSourceActions::GetName() const
+{
+	return LOCTEXT("RtspMediaSource_Name", "RTSP Media Source");
+}
+
+UClass* FRtspMediaSourceActions::GetSupportedClass() const
+{
+	return URtspMediaSource::StaticClass();
+}
+
+FColor FRtspMediaSourceActions::GetTypeColor() const
+{
+	return FColor(0, 180, 160);
+}
+
+#undef LOCTEXT_NAMESPACE
diff --git a/Engine/Plugins/Media/RTSPMedia/Source/RTSPMediaEditor/Private/RtspMediaSourceActions.h b/Engine/Plugins/Media/RTSPMedia/Source/RTSPMediaEditor/Private/RtspMediaSourceActions.h
new file mode 100644
--- /dev/null
+++ b/Engine/Plugins/Media/RTSPMedia/Source/RTSPMediaEditor/Private/RtspMediaSourceActions.h
@@ -0,0 +1,16 @@
+// Copyright Epic Games, Inc. All Rights Reserved.
+
+#pragma once
+
+#include "AssetTools/MediaSourceActions.h"
+
+class FRtspMediaSourceActions : public FMediaSourceActions
+{
+public:
+	//~ Begin FAssetTypeActions_Base
+	virtual bool CanFilter() override;
+	virtual FText GetName() const override;
+	virtual UClass* GetSupportedClass() const override;
+	virtual FColor GetTypeColor() const override;
+	//~ End FAssetTypeActions_Base
+};
diff --git a/Engine/Plugins/Media/RTSPMedia/Source/RTSPMediaEditor/Private/RtspMediaSourceFactory.cpp b/Engine/Plugins/Media/RTSPMedia/Source/RTSPMediaEditor/Private/RtspMediaSourceFactory.cpp
new file mode 100644
--- /dev/null
+++ b/Engine/Plugins/Media/RTSPMedia/Source/RTSPMediaEditor/Private/RtspMediaSourceFactory.cpp
@@ -0,0 +1,29 @@
+// Copyright Epic Games, Inc. All Rights Reserved.
+
+#include "RtspMediaSourceFactory.h"
+
+#include "AssetTypeCategories.h"
+#include "RtspMediaSource.h"
+
+URtspMediaSourceFactory::URtspMediaSourceFactory(const FObjectInitializer& ObjectInitializer)
+	: Super(ObjectInitializer)
+{
+	SupportedClass = URtspMediaSource::StaticClass();
+	bCreateNew = true;
+	bEditAfterNew = true;
+}
+
+UObject* URtspMediaSourceFactory::FactoryCreateNew(UClass* InClass, UObject* InParent, FName InName, EObjectFlags InFlags, UObject* InContext, FFeedbackContext* InWarn)
+{
+	return NewObject<URtspMediaSource>(InParent, InClass, InName, InFlags);
+}
+
+uint32 URtspMediaSourceFactory::GetMenuCategories() const
+{
+	return EAssetTypeCategories::Media;
+}
+
+bool URtspMediaSourceFactory::ShouldShowInNewMenu() const
+{
+	return true;
+}
\ No newline at end of file
diff --git a/Engine/Plugins/Media/RTSPMedia/Source/RTSPMediaEditor/Private/RtspMediaSourceFactory.h b/Engine/Plugins/Media/RTSPMedia/Source/RTSPMediaEditor/Private/RtspMediaSourceFactory.h
new file mode 100644
--- /dev/null
+++ b/Engine/Plugins/Media/RTSPMedia/Source/RTSPMediaEditor/Private/RtspMediaSourceFactory.h
@@ -0,0 +1,22 @@
+// Copyright Epic Games, Inc. All Rights Reserved.
+
+#pragma once
+
+#include "Factories/Factory.h"
+
+#include "RtspMediaSourceFactory.generated.h"
+
+UCLASS()
+class URtspMediaSourceFactory : public UFactory
+{
+	GENERATED_BODY()
+
+public:
+	URtspMediaSourceFactory(const FObjectInitializer& ObjectInitializer);
+
+	//~ Begin UFactory
+	virtual UObject* FactoryCreateNew(UClass* InClass, UObject* InParent, FName InName, EObjectFlags InFlags, UObject* InContext, FFeedbackContext* InWarn) override;
+	virtual uint32 GetMenuCategories() const override;
+	virtual bool ShouldShowInNewMenu() const override;
+	//~ End UFactory
+};
\ No newline at end of file
diff --git a/Engine/Plugins/Media/RTSPMedia/Source/RTSPMediaEditor/RTSPMediaEditor.Build.cs b/Engine/Plugins/Media/RTSPMedia/Source/RTSPMediaEditor/RTSPMediaEditor.Build.cs
new file mode 100644
--- /dev/null
+++ b/Engine/Plugins/Media/RTSPMedia/Source/RTSPMediaEditor/RTSPMediaEditor.Build.cs
@@ -0,0 +1,22 @@
+// Copyright Epic Games, Inc. All Rights Reserved.
+
+using UnrealBuildTool;
+
+public class RTSPMediaEditor : ModuleRules
+{
+    public RTSPMediaEditor(ReadOnlyTargetRules Target) : base(Target)
+    {
+        PrivateDependencyModuleNames.AddRange(new string[]
+        {
+            "AssetDefinition",
+            "Core",
+            "CoreUObject",
+            "MediaPlayerEditor",
+            "Projects",
+            "RTSPMedia",
+            "Slate",
+            "SlateCore",
+            "UnrealEd"
+        });
+    }
+}
\ No newline at end of file
