Remove NAL unit ordering check in RtpDecoder

The decoder's timestamp discontinuity check assured all NAL units in an access unit shared the same RTP timestamp.
Some non-compliant senders send SPS/PPS (H.264 metadata) NAL units in separate RTP packets with timestamps that differ from the image slice.
This caused valid image data to be discarded.

Frame boundaries are already correctly determined by the marker bit and slice type check in the decoder.
We also already protect against out of order packets via the prior depacketizer's sequence number validation.

#rnx

#Synced-CL 51260989

diff --git a/Engine/Plugins/Media/RTSPMedia/Source/RTSPMedia/Private/RTP/RtpDecoder.cpp b/Engine/Plugins/Media/RTSPMedia/Source/RTSPMedia/Private/RTP/RtpDecoder.cpp
--- a/Engine/Plugins/Media/RTSPMedia/Source/RTSPMedia/Private/RTP/RtpDecoder.cpp
+++ b/Engine/Plugins/Media/RTSPMedia/Source/RTSPMedia/Private/RTP/RtpDecoder.cpp
@@ -344,23 +344,6 @@ void FRtpDecoder::HandleOnH264NalUnit(TArray<uint8> InNalUnit, uint32 InTimestam
 		ResetFrameBuffer();
 	}
 
-	// Check for timestamp discontinuity in case NAL units have arrived out of order
-	// If we have data in our frame buffer from a previous NAL unit, check the timestamp
-	// of the new unit matches that of the frame as a whole.
-	if (FrameBuffer.Num() > 0 && InTimestamp != CurrentFrameTimestamp)
-	{
-		UE_LOG(LogRtspMedia, Warning, TEXT("Detected timestamp discontinuity, discarding frame"));
-		ResetFrameBuffer();
-		bWaitForIDRFrame = true;
-		return;
-	}
-
-	// Starting a new frame? Store the timestamp of this NAL unit as the timestamp for our frame
-	if (FrameBuffer.IsEmpty())
-	{
-		CurrentFrameTimestamp = InTimestamp;
-	}
-
 	// Append the NAL unit to our frame buffer
 	FrameBuffer.Append(InNalUnit);
 	
@@ -376,7 +359,7 @@ void FRtpDecoder::HandleOnH264NalUnit(TArray<uint8> InNalUnit, uint32 InTimestam
 	// Some senders will set the marker bit on NAL units that contain only SPS/PPS data (metadata, no image data).
 	// The decoder requires a complete access unit with an image slice or decoding will fail.
 	// So once the marker bit is set we need to check that we have received a NAL unit that represents an image slice.
-	// Any metadata units we received with the same timestamp will be aggregated into this access unit.
+	// Any metadata units we receive will be aggregated into the next access unit.
 	// With the marker bit set and bIsSlice == true, we should have a complete unit containing image data ready for decoding.
 	if (!bIsSlice)
 	{
@@ -597,7 +580,6 @@ void FRtpDecoder::PollDecodedOutput()
 void FRtpDecoder::ResetFrameBuffer()
 {
 	FrameBuffer.Empty();
-	CurrentFrameTimestamp = 0;
 }
 
 void FRtpDecoder::ResetPlatformVideoResource()
diff --git a/Engine/Plugins/Media/RTSPMedia/Source/RTSPMedia/Private/RTP/RtpDecoder.h b/Engine/Plugins/Media/RTSPMedia/Source/RTSPMedia/Private/RTP/RtpDecoder.h
--- a/Engine/Plugins/Media/RTSPMedia/Source/RTSPMedia/Private/RTP/RtpDecoder.h
+++ b/Engine/Plugins/Media/RTSPMedia/Source/RTSPMedia/Private/RTP/RtpDecoder.h
@@ -100,7 +100,6 @@ private:
 	uint32 DecoderPollIntervalMs = RtspMedia::Default::DecoderPollIntervalMs;
 
 	// Decoding State
-	uint32 CurrentFrameTimestamp = 0;
 	TArray<uint8> FrameBuffer;
 	TUniquePtr<FSdpSession> SdpSession;
 	uint32 VideoClockRate = 90000;
