Add adaptive jitter buffer depth estimation to RTSP media source

Previously the RTSP media source would use a static configuration for the jitter buffer depth.
By default the media source will now measure the network jitter and adjust the jitter buffer depth dynamically.
It's still possible to set a static jitter buffer depth if required.

#Synced-CL 51260989

diff --git a/Engine/Plugins/Media/RTSPMedia/Source/RTSPMedia/Private/Player/RtspMediaPlayer.cpp b/Engine/Plugins/Media/RTSPMedia/Source/RTSPMedia/Private/Player/RtspMediaPlayer.cpp
--- a/Engine/Plugins/Media/RTSPMedia/Source/RTSPMedia/Private/Player/RtspMediaPlayer.cpp
+++ b/Engine/Plugins/Media/RTSPMedia/Source/RTSPMedia/Private/Player/RtspMediaPlayer.cpp
@@ -67,10 +67,12 @@ bool FRtspMediaPlayer::Open(const FString& InUrl, const IMediaOptions* InOptions
 	const int64 MaxQueuedVideoSamplesOption = InOptions->GetMediaOption(RtspMedia::Option::MaxQueuedVideoSamples, static_cast<int64>(RtspMedia::Default::MaxQueuedVideoSamples));
 	Samples = MakeUnique<FMediaSamples>(RtspMedia::Default::MaxQueuedAudioSamples, static_cast<int32>(FMath::Clamp(MaxQueuedVideoSamplesOption, 2, 16)));
 
-	const double RequestTimeoutSecondsOption = InOptions->GetMediaOption(RtspMedia::Option::RequestTimeoutSeconds, static_cast<double>(RtspMedia::Default::RequestTimeoutSeconds));
+	const double RequestTimeoutSecondsOption = InOptions->GetMediaOption(RtspMedia::Option::RequestTimeoutSeconds, RtspMedia::Default::RequestTimeoutSeconds);
 	const int64 SocketBufferSizeKbOption = InOptions->GetMediaOption(RtspMedia::Option::SocketBufferSize, static_cast<int64>(RtspMedia::Default::SocketBufferSizeKb));
 	const int64 TransportProtocolOption = InOptions->GetMediaOption(RtspMedia::Option::TransportProtocol, static_cast<int64>(RtspMedia::Default::TransportProtocol));
 	const int64 JitterBufferDepthMsOption = InOptions->GetMediaOption(RtspMedia::Option::JitterBufferDepthMs, static_cast<int64>(RtspMedia::Default::JitterBufferDepthMs));
+	const bool bJitterBufferAutoAdjustOption = InOptions->GetMediaOption(RtspMedia::Option::JitterBufferAutoAdjust, RtspMedia::Default::bJitterBufferAutoAdjust);
+	const double JitterBufferObservationWindowSecondsOption = InOptions->GetMediaOption(RtspMedia::Option::JitterBufferObservationWindowSeconds, RtspMedia::Default::JitterBufferObservationWindowSeconds);
 
 	RtspClientConfiguration = MakeUnique<FRtspClientConfiguration>();
 	RtspClientConfiguration->Url = InUrl;
@@ -78,6 +80,8 @@ bool FRtspMediaPlayer::Open(const FString& InUrl, const IMediaOptions* InOptions
 	RtspClientConfiguration->SocketBufferSizeBytes = static_cast<int32>(FMath::Clamp(SocketBufferSizeKbOption, 128, 4096) * 1024);
 	RtspClientConfiguration->TransportProtocol = static_cast<FRtspTransportConfiguration::TransportProtocol>(FMath::Clamp(TransportProtocolOption, 0, 1));
 	RtspClientConfiguration->JitterBufferDepthMs = static_cast<uint32>(FMath::Clamp(JitterBufferDepthMsOption, 0, 5000));
+	RtspClientConfiguration->bJitterBufferAutoAdjust = bJitterBufferAutoAdjustOption;
+	RtspClientConfiguration->JitterBufferObservationWindowSeconds = static_cast<float>(FMath::Clamp(JitterBufferObservationWindowSecondsOption, 30.0, 300.0));
 
 	ResetRestartState();
 	bAutoRestart = InOptions->GetMediaOption(RtspMedia::Option::AutoReconnect, RtspMedia::Default::bAutoReconnect);
diff --git a/Engine/Plugins/Media/RTSPMedia/Source/RTSPMedia/Private/RTP/RtpJitterBuffer.cpp b/Engine/Plugins/Media/RTSPMedia/Source/RTSPMedia/Private/RTP/RtpJitterBuffer.cpp
--- a/Engine/Plugins/Media/RTSPMedia/Source/RTSPMedia/Private/RTP/RtpJitterBuffer.cpp
+++ b/Engine/Plugins/Media/RTSPMedia/Source/RTSPMedia/Private/RTP/RtpJitterBuffer.cpp
@@ -15,15 +15,14 @@ bool FRtpJitterBuffer::Initialize(const uint32 InBufferDepthMs, const uint32 InC
 	}
 
 	ClockRate = InClockRate;
-	PacketBufferDepthMs = InBufferDepthMs;
-	PacketBufferDepthSeconds = PacketBufferDepthMs / 1000.0;
+	PacketBufferDepthSeconds = InBufferDepthMs / 1000.0f;
 
 	ResetBuffer();
 	ResetPacing();
 
 	bInitialized = true;
-	
-	UE_LOG(LogRtspMedia, Verbose, TEXT("Initialized RTP jitter buffer with buffer depth (ms): %u Clock rate: %u"), PacketBufferDepthMs, ClockRate);
+
+	UE_LOG(LogRtspMedia, Verbose, TEXT("Initialized RTP jitter buffer with buffer depth: %.1fms Clock rate: %u"), PacketBufferDepthSeconds * 1000.0f, ClockRate);
 	return true;
 }
 
@@ -53,7 +52,7 @@ void FRtpJitterBuffer::EnqueuePacket(FRtpPacket InPacket)
 		TEXT("RTP jitter buffer enqueuing packet with sequence number: %u RTP timestamp: %u"),
 		InPacket.Header.SequenceNumber,
 		InPacket.Header.Timestamp);
-	
+
 	PacketBuffer.Emplace(MoveTemp(InPacket));
 
 	LastEnqueuedTimestamp = Timestamp;
@@ -66,7 +65,7 @@ bool FRtpJitterBuffer::DequeueIfReady(FRtpPacket& OutPacket)
 	{
 		return false;
 	}
-	
+
 	if (PacketBuffer.IsEmpty())
 	{
 		return false;
@@ -78,13 +77,13 @@ bool FRtpJitterBuffer::DequeueIfReady(FRtpPacket& OutPacket)
 		ResetPacing();
 		return false;
 	}
-	
+
 	const double Now = FPlatformTime::Seconds();
 
 	FRtpPacket& FirstPacket = PacketBuffer[QueueIndex];
 
 	// Skip buffer filling and pacing if the configured packet buffer depth is 0
-	if (PacketBufferDepthMs > 0)
+	if (PacketBufferDepthSeconds > 0.0f)
 	{
 		if (bFilling)
 		{
@@ -92,7 +91,7 @@ bool FRtpJitterBuffer::DequeueIfReady(FRtpPacket& OutPacket)
 			{
 				return false;
 			}
-		
+
 			if (Now - FillStartTimeSeconds < PacketBufferDepthSeconds)
 			{
 				return false;
@@ -108,34 +107,39 @@ bool FRtpJitterBuffer::DequeueIfReady(FRtpPacket& OutPacket)
 		if (!bSameTimestamp)
 		{
 			// If it's not release time yet, abort.
-			if (EstimatedFrameIntervalSeconds > 0.0 && Now < NextReleaseTimeSeconds)
+			if (EstimatedFrameIntervalSeconds > 0.0f && Now < NextReleaseTimeSeconds)
 			{
 				return false;
 			}
 
-			// Calculate the initial estimated frame interval and target frame count if not set 
-			if (EstimatedFrameIntervalSeconds == 0.0 && QueuedFrameCount >= 2)
+			// Calculate the initial estimated frame interval and target frame count if not set
+			if (EstimatedFrameIntervalSeconds == 0.0f && QueuedFrameCount >= 2)
 			{
 				// The cast to int32 is intentional for wraparound safety.
 				const int32 Span = static_cast<int32>(PacketBuffer.Last().Header.Timestamp - FirstPacket.Header.Timestamp);
-				EstimatedFrameIntervalSeconds = (static_cast<double>(Span) / ClockRate) / (QueuedFrameCount - 1);
+				// A non-positive span can occur if all buffered packets share the same RTP timestamp.
+				// Skip estimation and defer to the next dequeue attempt when more data is available.
+				if (Span > 0)
+				{
+					EstimatedFrameIntervalSeconds = (static_cast<float>(Span) / ClockRate) / (QueuedFrameCount - 1);
+				}
 			}
 
 			// Adjust the estimated frame interval for any deviation from the previously calculated frame interval
-			if (EstimatedFrameIntervalSeconds > 0.0)
+			if (EstimatedFrameIntervalSeconds > 0.0f)
 			{
-				const double TargetFrameCount = PacketBufferDepthSeconds / EstimatedFrameIntervalSeconds;
-				
+				const float TargetFrameCount = PacketBufferDepthSeconds / EstimatedFrameIntervalSeconds;
+
 				// Normalized deviation of actual queued frame count vs the target frame count
-				const double Error = (static_cast<double>(QueuedFrameCount) - TargetFrameCount) / TargetFrameCount;
+				const float Error = (static_cast<float>(QueuedFrameCount) - TargetFrameCount) / TargetFrameCount;
 
 				// Prevent the adjustment from becoming negative when an extreme burst of packets is present in the buffer.
-				const double Adjustment = FMath::Max(1.0 - BufferLevelGain * Error, BufferLevelGain);
-				
+				const float Adjustment = FMath::Max(1.0f - BufferLevelGain * Error, BufferLevelGain);
+
 				// Adjust the estimated frame interval by the error to bring it closer to the actual frame interval
 				EstimatedFrameIntervalSeconds *= Adjustment;
 			}
-			
+
 			// We will release now so compute the next release time
 			NextReleaseTimeSeconds += EstimatedFrameIntervalSeconds;
 		}
@@ -145,14 +149,14 @@ bool FRtpJitterBuffer::DequeueIfReady(FRtpPacket& OutPacket)
 	{
 		--QueuedFrameCount;
 	}
-	
+
 	UE_LOG(
 		LogRtspMedia,
-		VeryVerbose, 
+		VeryVerbose,
 		TEXT("Dequeuing RTP packet from jitter buffer. Sequence Number: %u Timestamp: %u"),
 		FirstPacket.Header.SequenceNumber,
 		FirstPacket.Header.Timestamp);
-		
+
 	bFirstPacketDequeued = true;
 	LastDequeuedTimestamp = FirstPacket.Header.Timestamp;
 	OutPacket = MoveTemp(FirstPacket);
@@ -162,6 +166,18 @@ bool FRtpJitterBuffer::DequeueIfReady(FRtpPacket& OutPacket)
 	return true;
 }
 
+void FRtpJitterBuffer::SetTargetBufferDepth(float InDepthSeconds)
+{
+	const float PreviousDepthSeconds = PacketBufferDepthSeconds;
+	PacketBufferDepthSeconds = FMath::Max(InDepthSeconds, 0.0f);
+
+	if (!FMath::IsNearlyEqual(PacketBufferDepthSeconds, PreviousDepthSeconds, 0.001f))
+	{
+		UE_LOG(LogRtspMedia, Verbose, TEXT("Jitter buffer depth updated: %.1fms -> %.1fms"),
+			PreviousDepthSeconds * 1000.0f, PacketBufferDepthSeconds * 1000.0f);
+	}
+}
+
 int32 FRtpJitterBuffer::NumQueuedPackets() const
 {
 	return PacketBuffer.Num() - QueueIndex;
@@ -190,7 +206,7 @@ void FRtpJitterBuffer::ResetBuffer()
 
 void FRtpJitterBuffer::ResetPacing()
 {
-	EstimatedFrameIntervalSeconds = 0.0;
+	EstimatedFrameIntervalSeconds = 0.0f;
 	bFilling = true;
 	FillStartTimeSeconds = 0.0;
 	bFirstPacketDequeued = false;
diff --git a/Engine/Plugins/Media/RTSPMedia/Source/RTSPMedia/Private/RTP/RtpJitterBuffer.h b/Engine/Plugins/Media/RTSPMedia/Source/RTSPMedia/Private/RTP/RtpJitterBuffer.h
--- a/Engine/Plugins/Media/RTSPMedia/Source/RTSPMedia/Private/RTP/RtpJitterBuffer.h
+++ b/Engine/Plugins/Media/RTSPMedia/Source/RTSPMedia/Private/RTP/RtpJitterBuffer.h
@@ -16,6 +16,7 @@ public:
 	bool Initialize(const uint32 InBufferDepthMs, const uint32 InClockRate);
 	void EnqueuePacket(FRtpPacket InPacket);
 	bool DequeueIfReady(FRtpPacket& OutPacket);
+	void SetTargetBufferDepth(float InDepthSeconds);
 
 private:
 	void CompactPacketBuffer();
@@ -23,14 +24,13 @@ private:
 	void ResetBuffer();
 	void ResetPacing();
 
-	static constexpr double BufferLevelGain = 0.02;
+	static constexpr float BufferLevelGain = 0.02f;
 
 	// Configuration
 	uint32 ClockRate = 0;
 	bool bInitialized = false;
-	uint32 PacketBufferDepthMs = 0;
-	double PacketBufferDepthSeconds = 0.0;
-	
+	float PacketBufferDepthSeconds = 0.0f;
+
 	// Buffer
 	TArray<FRtpPacket> PacketBuffer;
 	int32 QueueIndex = 0;
@@ -40,7 +40,7 @@ private:
 	double FillStartTimeSeconds = 0.0;
 
 	// Frame Pacing
-	double EstimatedFrameIntervalSeconds = 0.0;
+	float EstimatedFrameIntervalSeconds = 0.0f;
 	double NextReleaseTimeSeconds = 0.0;
 	int32 QueuedFrameCount = 0;
 
diff --git a/Engine/Plugins/Media/RTSPMedia/Source/RTSPMedia/Private/RTP/RtpJitterEstimator.cpp b/Engine/Plugins/Media/RTSPMedia/Source/RTSPMedia/Private/RTP/RtpJitterEstimator.cpp
new file mode 100644
--- /dev/null
+++ b/Engine/Plugins/Media/RTSPMedia/Source/RTSPMedia/Private/RTP/RtpJitterEstimator.cpp
@@ -0,0 +1,116 @@
+// Copyright Epic Games, Inc. All Rights Reserved.
+
+#include "RtpJitterEstimator.h"
+
+#include "RtspMediaConstants.h"
+
+#include "HAL/PlatformTime.h"
+
+bool FRtpJitterEstimator::Initialize(uint32 InClockRate, float InObservationWindowSeconds)
+{
+	if (InClockRate == 0)
+	{
+		UE_LOG(LogRtspMedia, Error, TEXT("Attempted to initialize jitter estimator with a zero clock rate"));
+		return false;
+	}
+
+	if (InObservationWindowSeconds <= 0.0f)
+	{
+		UE_LOG(LogRtspMedia, Error, TEXT("Attempted to initialize jitter estimator with a non-positive observation window"));
+		return false;
+	}
+
+	ClockRate = InClockRate;
+	ObservationWindowSeconds = InObservationWindowSeconds;
+	BaselineSeconds = 0.0;
+	RecentMaxLatenessSeconds = 0.0;
+	WindowStartTimeSeconds = 0.0;
+	LastTimestamp = 0;
+	LastArrivalTimeSeconds = 0.0;
+	bFirstFrame = true;
+	bInitialized = true;
+
+	UE_LOG(LogRtspMedia, Verbose, TEXT("Initialized jitter estimator with clock rate: %u observation window: %.1fs"), ClockRate, ObservationWindowSeconds);
+	return true;
+}
+
+void FRtpJitterEstimator::RecordArrival(uint32 InRtpTimestamp)
+{
+	if (!bInitialized)
+	{
+		return;
+	}
+
+	// Skip FU-A fragments (same RTP timestamp as previous)
+	if (!bFirstFrame && InRtpTimestamp == LastTimestamp)
+	{
+		return;
+	}
+
+	const double Now = FPlatformTime::Seconds();
+
+	if (bFirstFrame)
+	{
+		LastTimestamp = InRtpTimestamp;
+		LastArrivalTimeSeconds = Now;
+		WindowStartTimeSeconds = Now;
+		bFirstFrame = false;
+		return;
+	}
+
+	// Expected interval from RTP timestamps
+	// int32 cast handles RTP timestamp wraparound
+	const double ExpectedInterval = static_cast<double>(static_cast<int32>(InRtpTimestamp - LastTimestamp)) / ClockRate;
+
+	// Negative expected interval indicates a B-frame (presentation timestamp earlier than
+	// previous frame, but delivered in decode order). Skip the measurement — timestamp
+	// reordering is not network jitter.
+	if (ExpectedInterval <= 0.0)
+	{
+		return;
+	}
+
+	const double ActualInterval = Now - LastArrivalTimeSeconds;
+	const double Lateness = ActualInterval - ExpectedInterval;
+
+	// Immediate step-up on any spike exceeding baseline
+	if (Lateness > BaselineSeconds)
+	{
+		BaselineSeconds = CeilToGranularity(Lateness);
+		WindowStartTimeSeconds = Now;
+		RecentMaxLatenessSeconds = 0.0;
+	}
+
+	// Track the worst lateness in the current window
+	RecentMaxLatenessSeconds = FMath::Max(RecentMaxLatenessSeconds, FMath::Max(Lateness, 0.0));
+
+	// At window boundary: evaluate and possibly step down
+	if ((Now - WindowStartTimeSeconds) >= ObservationWindowSeconds)
+	{
+		if (RecentMaxLatenessSeconds < BaselineSeconds)
+		{
+			BaselineSeconds = CeilToGranularity(RecentMaxLatenessSeconds);
+		}
+
+		// Reset window
+		WindowStartTimeSeconds = Now;
+		RecentMaxLatenessSeconds = 0.0;
+	}
+
+	LastTimestamp = InRtpTimestamp;
+	LastArrivalTimeSeconds = Now;
+}
+
+float FRtpJitterEstimator::GetTargetBufferDepthSeconds() const
+{
+	return static_cast<float>(CeilToGranularity(BaselineSeconds * HeadroomMultiplier));
+}
+
+double FRtpJitterEstimator::CeilToGranularity(double InSeconds)
+{
+	if (InSeconds <= 0.0)
+	{
+		return 0.0;
+	}
+	return FMath::CeilToDouble(InSeconds / GranularitySeconds) * GranularitySeconds;
+}
diff --git a/Engine/Plugins/Media/RTSPMedia/Source/RTSPMedia/Private/RTP/RtpJitterEstimator.h b/Engine/Plugins/Media/RTSPMedia/Source/RTSPMedia/Private/RTP/RtpJitterEstimator.h
new file mode 100644
--- /dev/null
+++ b/Engine/Plugins/Media/RTSPMedia/Source/RTSPMedia/Private/RTP/RtpJitterEstimator.h
@@ -0,0 +1,49 @@
+// Copyright Epic Games, Inc. All Rights Reserved.
+
+#pragma once
+
+#include "CoreMinimal.h"
+#include "RtspMediaDefaults.h"
+
+/**
+ * Observes RTP packet arrival timing to estimate the ideal jitter buffer depth.
+ *
+ * The estimator compares expected inter-frame intervals (derived from RTP timestamps)
+ * against actual wall-clock delivery intervals. The difference ("lateness") reflects
+ * network jitter. A baseline buffer depth is maintained that steps up immediately
+ * on spikes and steps down conservatively after sustained improvement over observation
+ * windows.
+ *
+ * Start at zero: baseline begins at 0 (passthrough). Clean networks stay at zero latency.
+ * Spikes drive the baseline up immediately.
+ * At window boundaries, baseline clamps down to the worst observed lateness in that window.
+ * Baseline values are rounded up to a 10ms granularity to avoid minor oscillations.
+ *
+ * GetTargetBufferDepthSeconds() returns the baseline with a headroom multiplier applied.
+ * During a stall, the buffer drains at the playback rate — without headroom the buffer
+ * empties right as the burst arrives, causing a visible stutter.
+ */
+class FRtpJitterEstimator
+{
+public:
+	bool Initialize(uint32 InClockRate, float InObservationWindowSeconds);
+	void RecordArrival(uint32 InRtpTimestamp);
+	float GetTargetBufferDepthSeconds() const;
+
+private:
+	/** Round up to the nearest 10ms. Returns 0 for values at or below 0. */
+	static double CeilToGranularity(double InSeconds);
+
+	static constexpr double GranularitySeconds = 0.01;
+	static constexpr double HeadroomMultiplier = 2.0;
+
+	uint32 ClockRate = 0;
+	bool bInitialized = false;
+	double BaselineSeconds = 0.0;
+	float ObservationWindowSeconds = RtspMedia::Default::JitterBufferObservationWindowSeconds;
+	double RecentMaxLatenessSeconds = 0.0;
+	double WindowStartTimeSeconds = 0.0;
+	uint32 LastTimestamp = 0;
+	double LastArrivalTimeSeconds = 0.0;
+	bool bFirstFrame = true;
+};
diff --git a/Engine/Plugins/Media/RTSPMedia/Source/RTSPMedia/Private/RTSP/RtspClient.cpp b/Engine/Plugins/Media/RTSPMedia/Source/RTSPMedia/Private/RTSP/RtspClient.cpp
--- a/Engine/Plugins/Media/RTSPMedia/Source/RTSPMedia/Private/RTSP/RtspClient.cpp
+++ b/Engine/Plugins/Media/RTSPMedia/Source/RTSPMedia/Private/RTSP/RtspClient.cpp
@@ -4,6 +4,7 @@
 
 #include "RtspMediaConstants.h"
 
+#include "RTP/RtpJitterEstimator.h"
 #include "RTP/RtpParser.h"
 #include "RTP/RtpPacket.h"
 
@@ -85,12 +86,27 @@ void FRtspClient::Disconnect()
 
 bool FRtspClient::InitializeJitterBuffer(const uint32 InClockRate)
 {
+	uint32 BufferDepthMs = Configuration.JitterBufferDepthMs;
+
+	if (Configuration.bJitterBufferAutoAdjust)
+	{
+		JitterEstimator = MakeUnique<FRtpJitterEstimator>();
+		if (!JitterEstimator->Initialize(InClockRate, Configuration.JitterBufferObservationWindowSeconds))
+		{
+			JitterEstimator.Reset();
+			return false;
+		}
+		BufferDepthMs = 0;
+	}
+
 	JitterBuffer = MakeUnique<FRtpJitterBuffer>();
-	if (!JitterBuffer->Initialize(Configuration.JitterBufferDepthMs, InClockRate))
+	if (!JitterBuffer->Initialize(BufferDepthMs, InClockRate))
 	{
 		JitterBuffer.Reset();
+		JitterEstimator.Reset();
 		return false;
 	}
+
 	return true;
 }
 
@@ -198,6 +214,11 @@ uint32 FRtspClient::Run()
 		// Check for waiting packets in the jitter buffer
 		if (JitterBuffer.IsValid())
 		{
+			if (JitterEstimator.IsValid())
+			{
+				JitterBuffer->SetTargetBufferDepth(JitterEstimator->GetTargetBufferDepthSeconds());
+			}
+
 			FRtpPacket BufferedPacket;
 			while (JitterBuffer->DequeueIfReady(BufferedPacket))
 			{
@@ -445,6 +466,11 @@ void FRtspClient::ProcessInterleavedPacket()
 		TArrayView<const uint8> RtpPayload;
 		if (FRtpParser::Parse(InterleavedPayload, Header, RtpPayload))
 		{
+			if (JitterEstimator.IsValid())
+			{
+				JitterEstimator->RecordArrival(Header.Timestamp);
+			}
+
 			if (JitterBuffer.IsValid())
 			{
 				FRtpPacket Packet{Header, TArray<uint8>(RtpPayload)};
@@ -987,6 +1013,7 @@ void FRtspClient::TeardownConnection()
 	SdpSession.Reset();
 
 	JitterBuffer.Reset();
+	JitterEstimator.Reset();
 	RequestTracker.Clear();
 
 	ReceiveBuffer.Reset();
diff --git a/Engine/Plugins/Media/RTSPMedia/Source/RTSPMedia/Private/RTSP/RtspClient.h b/Engine/Plugins/Media/RTSPMedia/Source/RTSPMedia/Private/RTSP/RtspClient.h
--- a/Engine/Plugins/Media/RTSPMedia/Source/RTSPMedia/Private/RTSP/RtspClient.h
+++ b/Engine/Plugins/Media/RTSPMedia/Source/RTSPMedia/Private/RTSP/RtspClient.h
@@ -3,6 +3,7 @@
 #pragma once
 
 #include "RTP/RtpJitterBuffer.h"
+#include "RTP/RtpJitterEstimator.h"
 #include "RtspMediaConstants.h"
 #include "RtspMediaDefaults.h"
 #include "RtspMessage.h"
@@ -21,7 +22,9 @@ struct FSdpSession;
 
 struct FRtspClientConfiguration
 {
+	bool bJitterBufferAutoAdjust = RtspMedia::Default::bJitterBufferAutoAdjust;
 	uint32 JitterBufferDepthMs = RtspMedia::Default::JitterBufferDepthMs;
+	float JitterBufferObservationWindowSeconds = RtspMedia::Default::JitterBufferObservationWindowSeconds;
 	int32 SocketBufferSizeBytes = RtspMedia::Default::SocketBufferSizeBytes;
 	float RequestTimeoutSeconds = RtspMedia::Default::RequestTimeoutSeconds;
 	FRtspTransportConfiguration::TransportProtocol TransportProtocol = FRtspTransportConfiguration::TransportProtocol::TCP;
@@ -147,4 +150,5 @@ private:
 
 	// Network jitter handling
 	TUniquePtr<FRtpJitterBuffer> JitterBuffer;
+	TUniquePtr<FRtpJitterEstimator> JitterEstimator;
 };
diff --git a/Engine/Plugins/Media/RTSPMedia/Source/RTSPMedia/Private/RtspMediaConstants.h b/Engine/Plugins/Media/RTSPMedia/Source/RTSPMedia/Private/RtspMediaConstants.h
--- a/Engine/Plugins/Media/RTSPMedia/Source/RTSPMedia/Private/RtspMediaConstants.h
+++ b/Engine/Plugins/Media/RTSPMedia/Source/RTSPMedia/Private/RtspMediaConstants.h
@@ -18,7 +18,9 @@ namespace RtspMedia
 		inline const FName AutoReconnect("AutoReconnect");
 		inline const FName DecoderBufferSize("DecoderBufferSize");
 		inline const FName DecoderPollIntervalMs("DecoderPollIntervalMs");
+		inline const FName JitterBufferAutoAdjust("JitterBufferAutoAdjust");
 		inline const FName JitterBufferDepthMs("JitterBufferDepthMs");
+		inline const FName JitterBufferObservationWindowSeconds("JitterBufferObservationWindowSeconds");
 		inline const FName MaxFragmentBufferSize("MaxFragmentBufferSize");
 		inline const FName MaxQueuedVideoSamples("MaxQueuedVideoSamples");
 		inline const FName MaxReconnectAttempts("MaxReconnectAttempts");
diff --git a/Engine/Plugins/Media/RTSPMedia/Source/RTSPMedia/Private/RtspMediaSource.cpp b/Engine/Plugins/Media/RTSPMedia/Source/RTSPMedia/Private/RtspMediaSource.cpp
--- a/Engine/Plugins/Media/RTSPMedia/Source/RTSPMedia/Private/RtspMediaSource.cpp
+++ b/Engine/Plugins/Media/RTSPMedia/Source/RTSPMedia/Private/RtspMediaSource.cpp
@@ -40,6 +40,11 @@ bool URtspMediaSource::GetMediaOption(const FName& InKey, bool InDefaultValue) c
 		return bAutoReconnect;
 	}
 	
+	if (InKey == RtspMedia::Option::JitterBufferAutoAdjust)
+	{
+		return bJitterBufferAutoAdjust;
+	}
+
 	return Super::GetMediaOption(InKey, InDefaultValue);
 }
 
@@ -95,6 +100,11 @@ int64 URtspMediaSource::GetMediaOption(const FName& InKey, int64 InDefaultValue)
 
 double URtspMediaSource::GetMediaOption(const FName& InKey, double InDefaultValue) const
 {
+	if (InKey == RtspMedia::Option::JitterBufferObservationWindowSeconds)
+	{
+		return JitterBufferObservationWindowSeconds;
+	}
+
 	if (InKey == RtspMedia::Option::MaxReconnectDelaySeconds)
 	{
 		return MaxReconnectDelaySeconds;
@@ -113,23 +123,25 @@ double URtspMediaSource::GetMediaOption(const FName& InKey, double InDefaultValu
 	return Super::GetMediaOption(InKey, InDefaultValue);
 }
 
-bool URtspMediaSource::HasMediaOption(const FName& Key) const
+bool URtspMediaSource::HasMediaOption(const FName& InKey) const
 {
-	if (Key == RtspMedia::Option::AutoReconnect ||
-		Key == RtspMedia::Option::DecoderBufferSize ||
-		Key == RtspMedia::Option::DecoderPollIntervalMs ||
-		Key == RtspMedia::Option::JitterBufferDepthMs ||
-		Key == RtspMedia::Option::MaxFragmentBufferSize ||
-		Key == RtspMedia::Option::MaxQueuedVideoSamples ||
-		Key == RtspMedia::Option::MaxReconnectAttempts ||
-		Key == RtspMedia::Option::MaxReconnectDelaySeconds ||
-		Key == RtspMedia::Option::MinReconnectDelaySeconds ||
-		Key == RtspMedia::Option::RequestTimeoutSeconds ||
-		Key == RtspMedia::Option::SocketBufferSize ||
-		Key == RtspMedia::Option::TransportProtocol)
+	if (InKey == RtspMedia::Option::AutoReconnect ||
+		InKey == RtspMedia::Option::DecoderBufferSize ||
+		InKey == RtspMedia::Option::DecoderPollIntervalMs ||
+		InKey == RtspMedia::Option::JitterBufferAutoAdjust ||
+		InKey == RtspMedia::Option::JitterBufferDepthMs ||
+		InKey == RtspMedia::Option::JitterBufferObservationWindowSeconds ||
+		InKey == RtspMedia::Option::MaxFragmentBufferSize ||
+		InKey == RtspMedia::Option::MaxQueuedVideoSamples ||
+		InKey == RtspMedia::Option::MaxReconnectAttempts ||
+		InKey == RtspMedia::Option::MaxReconnectDelaySeconds ||
+		InKey == RtspMedia::Option::MinReconnectDelaySeconds ||
+		InKey == RtspMedia::Option::RequestTimeoutSeconds ||
+		InKey == RtspMedia::Option::SocketBufferSize ||
+		InKey == RtspMedia::Option::TransportProtocol)
 	{
 		return true;
 	}
 
-	return Super::HasMediaOption(Key);
+	return Super::HasMediaOption(InKey);
 }
diff --git a/Engine/Plugins/Media/RTSPMedia/Source/RTSPMedia/Public/RtspMediaDefaults.h b/Engine/Plugins/Media/RTSPMedia/Source/RTSPMedia/Public/RtspMediaDefaults.h
--- a/Engine/Plugins/Media/RTSPMedia/Source/RTSPMedia/Public/RtspMediaDefaults.h
+++ b/Engine/Plugins/Media/RTSPMedia/Source/RTSPMedia/Public/RtspMediaDefaults.h
@@ -9,7 +9,9 @@ namespace RtspMedia::Default
 	inline constexpr bool bAutoReconnect = true;
 	inline constexpr int32 DecoderBufferSize = 5;
 	inline constexpr int32 DecoderPollIntervalMs = 5;
+	inline constexpr bool bJitterBufferAutoAdjust = true;
 	inline constexpr uint32 JitterBufferDepthMs = 1000;
+	inline constexpr float JitterBufferObservationWindowSeconds = 60.0f;
 	inline constexpr int32 MaxFragmentBufferSizeMb = 8;
 	inline constexpr int64 MaxFragmentBufferSizeBytes = MaxFragmentBufferSizeMb * 1024 * 1024;
 	inline constexpr int32 MaxQueuedVideoSamples = 4;
diff --git a/Engine/Plugins/Media/RTSPMedia/Source/RTSPMedia/Public/RtspMediaSource.h b/Engine/Plugins/Media/RTSPMedia/Source/RTSPMedia/Public/RtspMediaSource.h
--- a/Engine/Plugins/Media/RTSPMedia/Source/RTSPMedia/Public/RtspMediaSource.h
+++ b/Engine/Plugins/Media/RTSPMedia/Source/RTSPMedia/Public/RtspMediaSource.h
@@ -28,76 +28,85 @@ class RTSPMEDIA_API URtspMediaSource : public UBaseMediaSource
 public:
 
 	/** Host name or IP address of the RTSP server */
-	UPROPERTY(EditAnywhere, BlueprintReadWrite, Category = "Network")
+	UPROPERTY(EditAnywhere, BlueprintReadWrite, Category = "Server")
 	FString Host;
 
 	/** Port number of the RTSP Server */
-	UPROPERTY(EditAnywhere, BlueprintReadWrite, Category = "Network", meta=(ClampMin="1", ClampMax="65535"))
+	UPROPERTY(EditAnywhere, BlueprintReadWrite, Category = "Server", meta=(ClampMin="1", ClampMax="65535"))
 	int32 Port = RtspMedia::Default::Port;
 
 	/** Path to the stream on RTSP server */
-	UPROPERTY(EditAnywhere, BlueprintReadWrite, Category = "Network")
+	UPROPERTY(EditAnywhere, BlueprintReadWrite, Category = "Server")
 	FString Path;
 
 	/** Only TCP interleaved transport is supported */
-	UPROPERTY(BlueprintReadOnly, Category="Network", meta=(Hidden))
+	UPROPERTY(BlueprintReadOnly, Category="Client", meta=(Hidden))
 	ERtspMediaTransportProtocol TransportProtocol = RtspMedia::Default::TransportProtocol;
 	
 	/** Socket buffer size for stream data in KB.
 	 * Increase this setting for high bitrate streams (e.g. 4K) */
-	UPROPERTY(EditAnywhere, BlueprintReadWrite, Category="Network", AdvancedDisplay, meta=(DisplayName="Socket Buffer Size (KB)", ClampMin="128", ClampMax="4096"))
+	UPROPERTY(EditAnywhere, BlueprintReadWrite, Category="Client", meta=(DisplayName="Socket Buffer Size (KB)", ClampMin="128", ClampMax="4096"))
 	int32 SocketBufferSizeKb = RtspMedia::Default::SocketBufferSizeKb;
-
-	/**
-	 * The jitter buffer accounts for spikes in the network delivery of RTP packets in order to provide a steady stream of packets to the decoder.
-	 * A higher buffer depth will absorb more network jitter at the cost of additional latency.
-	 * On clean networks a low or zero value may be configured, and will minimise the latency introduced by the buffer. */
-	UPROPERTY(EditAnywhere, BlueprintReadWrite, Category="Network", AdvancedDisplay, meta=(DisplayName="Jitter Buffer Depth (ms)", ClampMin="0", ClampMax="5000"))
-	int32 JitterBufferDepthMs = RtspMedia::Default::JitterBufferDepthMs;
 	
 	/** RTSP request timeout in seconds.
 	 * Increase for slow or unreliable networks to avoid premature session failures */
-	UPROPERTY(EditAnywhere, BlueprintReadWrite, Category="Network", AdvancedDisplay, meta=(DisplayName="Request Timeout (s)", ClampMin="1.0", ClampMax="30.0"))
+	UPROPERTY(EditAnywhere, BlueprintReadWrite, Category="Client", meta=(DisplayName="Request Timeout (s)", ClampMin="1.0", ClampMax="30.0"))
 	float RequestTimeoutSeconds = RtspMedia::Default::RequestTimeoutSeconds;
 	
 	/** Automatically reconnect when the connection fails
 	 * The player will begin reconnection attempts after the delay specified in Min Retry Delay */
-	UPROPERTY(EditAnywhere, BlueprintReadWrite, Category="Network", AdvancedDisplay)
+	UPROPERTY(EditAnywhere, BlueprintReadWrite, Category="Client")
 	bool bAutoReconnect = RtspMedia::Default::bAutoReconnect;
 
 	/** The minimum amount of time to wait before attempting to reconnect after a connection failure */
-	UPROPERTY(EditAnywhere, BlueprintReadWrite, Category="Network", AdvancedDisplay, meta=(DisplayName="Min Reconnect Delay (s)", EditCondition="bAutoReconnect", ClampMin="1.0", ClampMax="60.0"))
+	UPROPERTY(EditAnywhere, BlueprintReadWrite, Category="Client", meta=(DisplayName="Min Reconnect Delay (s)", EditCondition="bAutoReconnect", ClampMin="1.0", ClampMax="60.0"))
 	float MinReconnectDelaySeconds = RtspMedia::Default::MinReconnectDelaySeconds;
 
 	/** The maximum amount of time to wait before attempting to reconnect after a connection failure
 	 * The player will gradually increase the delay before each connection attempt but for no longer than this maximum */
-	UPROPERTY(EditAnywhere, BlueprintReadWrite, Category="Network", AdvancedDisplay, meta=(DisplayName="Max Reconnect Delay (s)", EditCondition="bAutoReconnect", ClampMin="1.0", ClampMax="300.0"))
+	UPROPERTY(EditAnywhere, BlueprintReadWrite, Category="Client", meta=(DisplayName="Max Reconnect Delay (s)", EditCondition="bAutoReconnect", ClampMin="1.0", ClampMax="300.0"))
 	float MaxReconnectDelaySeconds = RtspMedia::Default::MaxReconnectDelaySeconds;
 
 	/** The maximum number of times the player will attempt to reconnect after a connection failure
 	 * Set to 0 to retry indefinitely */
-	UPROPERTY(EditAnywhere, BlueprintReadWrite, Category="Network", AdvancedDisplay, meta=(EditCondition="bAutoReconnect", ClampMin="0", ClampMax="1000"))
+	UPROPERTY(EditAnywhere, BlueprintReadWrite, Category="Client", meta=(EditCondition="bAutoReconnect", ClampMin="0", ClampMax="1000"))
 	int32 MaxReconnectAttempts = RtspMedia::Default::MaxReconnectAttempts;
+	
+	/** The jitter buffer accounts for spikes in the network delivery of RTP packets in order to provide a steady stream of packets to the decoder. 
+	 * This setting automatically adjusts the jitter buffer depth based on network conditions.
+	 * When disabled, the static Depth value is used. */
+	UPROPERTY(EditAnywhere, BlueprintReadWrite, Category="Jitter Buffer", meta=(DisplayName="Auto Adjust"))
+	bool bJitterBufferAutoAdjust = RtspMedia::Default::bJitterBufferAutoAdjust;
+
+	/** A higher buffer depth will absorb more network jitter at the cost of additional latency.
+	 * On clean networks a low or zero value may be configured which will minimize the latency introduced by the buffer. */
+	UPROPERTY(EditAnywhere, BlueprintReadWrite, Category="Jitter Buffer", meta=(DisplayName="Depth (ms)", EditCondition="!bJitterBufferAutoAdjust", EditConditionHides, ClampMin="0", ClampMax="5000"))
+	int32 JitterBufferDepthMs = RtspMedia::Default::JitterBufferDepthMs;
+
+	/** How long network conditions must remain stable before the jitter buffer depth is reduced.
+	 * Longer windows provide more protection against intermittent spikes. */
+	UPROPERTY(EditAnywhere, BlueprintReadWrite, Category="Jitter Buffer", meta=(DisplayName="Observation Window (s)", EditCondition="bJitterBufferAutoAdjust", EditConditionHides, ClampMin="30.0", ClampMax="300.0"))
+	float JitterBufferObservationWindowSeconds = RtspMedia::Default::JitterBufferObservationWindowSeconds;
+	
+	/** Places a limit on how much memory is used for fragmented NAL unit reassembly
+	 * NAL units will be dropped if this memory limit is exceeded
+	 * Consider increasing this value for high bitrate streams */ 
+	UPROPERTY(EditAnywhere, BlueprintReadWrite, Category = "Depacketization", meta=(DisplayName="Max Fragment Buffer Size (MB)", ClampMin="2", ClampMax="64"))
+	int32 MaxFragmentBufferSizeMb = RtspMedia::Default::MaxFragmentBufferSizeMb;
 
 	/** Decoder reorder buffer size. Use 1 for minimum latency (i.e. no B-frames)
 	 * Increase if video frames appears out of order */
-	UPROPERTY(EditAnywhere, BlueprintReadWrite, Category = "Video", meta=(DisplayName="Decoder Buffer Size (frames)", ClampMin="1", ClampMax="16"))
+	UPROPERTY(EditAnywhere, BlueprintReadWrite, Category = "Decoder", meta=(DisplayName="Buffer Size (frames)", ClampMin="1", ClampMax="16"))
 	int32 DecoderBufferSize = RtspMedia::Default::DecoderBufferSize;
 
 	/** How frequently to poll the decoder for output frames (in milliseconds)
 	 * Lower values reduce latency but increase CPU usage */
-	UPROPERTY(EditAnywhere, BlueprintReadWrite, Category = "Video", meta=(DisplayName="Decoder Poll Interval (ms)", ClampMin="1", ClampMax="10"))
+	UPROPERTY(EditAnywhere, BlueprintReadWrite, Category = "Decoder", meta=(DisplayName="Poll Interval (ms)", ClampMin="1", ClampMax="10"))
 	int32 DecoderPollIntervalMs = RtspMedia::Default::DecoderPollIntervalMs;
 
-	/** Places a limit on how much memory is used for fragmented NAL unit reassembly
-	 * NAL units will be dropped if this memory limit is exceeded
-	 * Consider increasing this value for high bitrate streams */ 
-	UPROPERTY(EditAnywhere, BlueprintReadWrite, Category = "Video", meta=(DisplayName="Max Fragment Buffer Size (MB)", ClampMin="2", ClampMax="64"))
-	int32 MaxFragmentBufferSizeMb = RtspMedia::Default::MaxFragmentBufferSizeMb;
-
 	/** Maximum decoded video frames to queue while waiting for game thread consumption
 	 *  Higher values reduce visible skips after game thread hitches at the cost of memory */
-	UPROPERTY(EditAnywhere, BlueprintReadWrite, Category="Video", meta=(ClampMin="2", ClampMax="16"))
+	UPROPERTY(EditAnywhere, BlueprintReadWrite, Category="Decoder", meta=(ClampMin="2", ClampMax="16"))
 	int32 MaxQueuedVideoSamples = RtspMedia::Default::MaxQueuedVideoSamples;
 
 	//~ Begin UMediaSource
