[RTSPMedia] Extract H.264 SEI timecode from bitstream processor

#rnx
#Synced-CL 51630826

diff --git a/Engine/Plugins/Media/RTSPMedia/Source/RTSPMedia/Private/RTP/RtpDecoder.cpp b/Engine/Plugins/Media/RTSPMedia/Source/RTSPMedia/Private/RTP/RtpDecoder.cpp
--- a/Engine/Plugins/Media/RTSPMedia/Source/RTSPMedia/Private/RTP/RtpDecoder.cpp
+++ b/Engine/Plugins/Media/RTSPMedia/Source/RTSPMedia/Private/RTP/RtpDecoder.cpp
@@ -14,11 +14,14 @@
 #include "IElectraCodecFactory.h"
 #include "IElectraCodecFactoryModule.h"
 #include "IElectraDecoderOutputVideo.h"
+#include "IElectraDecoderFeaturesAndOptions.h"
 #include "MediaDecoderOutput.h"
 #include "ParameterDictionary.h"
-#include "Async/Async.h"
 #include "Utils/MPEG/ElectraUtilsMPEGVideo_H264.h"
+#include "Utils/MPEG/ElectraUtilsMPEGVideo.h"
+#include "Utils/VideoDecoderHelpers.h"
 
+#include "Async/Async.h"
 #include "HAL/Event.h"
 #include "HAL/PlatformProcess.h"
 #include "Misc/ScopeExit.h"
@@ -194,7 +197,7 @@ void FRtpDecoder::Stop()
  bool FRtpDecoder::CreateElectraDecoder(const TArray<uint8>& InSequenceParameterSet, const TArray<uint8>& InPictureParameterSet)
  {
 	// Build Codec Specific Data ('csd')
-	// Used to initalize the decoder and for each decode action.
+	// Used to initialize the decoder and for each decode action.
 	// Contains the sequence and picture parameter sets which include information like:
 	// SPS - Properties for the entire sequence of frames
 	// E.g. Resolution, profile, chroma format, bit depth
@@ -317,6 +320,7 @@ void FRtpDecoder::Stop()
 	}
 
 	BitstreamProcessor.Reset();
+	BitstreamInfoQueue.Empty();
 	BitstreamProcessor = ElectraDecoder->CreateBitstreamProcessor();
 
 	if (!BitstreamProcessor.IsValid())
@@ -423,6 +427,7 @@ void FRtpDecoder::HandleOnH264NalUnit(TArray<uint8> InNalUnit, uint32 InTimestam
 	{
 		case IElectraDecoder::EDecoderError::None:
 			bDecodedSuccessfully = true;
+			BitstreamInfoQueue.Enqueue(BitstreamInfo);
 			break;
 		case IElectraDecoder::EDecoderError::NoBuffer:
 			UE_LOG(LogRtspMedia, Verbose, TEXT("ElectraDecoder reports no buffer available, dropping unit"));
@@ -466,6 +471,12 @@ void FRtpDecoder::PollDecodedOutput()
 	while (ElectraDecoder->HaveOutput() == IElectraDecoder::EOutputStatus::Available)
 	{
 		TSharedPtr<IElectraDecoderOutput> Output = ElectraDecoder->GetOutput();
+
+		// Dequeue the corresponding bitstream info for this output frame.
+		// This must happen before any continue/discard paths to keep the queue in sync.
+		TSharedPtr<IElectraDecoderBitstreamInfo, ESPMode::ThreadSafe> FrameBitstreamInfo;
+		BitstreamInfoQueue.Dequeue(FrameBitstreamInfo);
+
 		if (!Output.IsValid())
 		{
 			UE_LOG(LogRtspMedia, Warning, TEXT("Received null output from electra decoder"));
@@ -556,6 +567,21 @@ void FRtpDecoder::PollDecodedOutput()
 		BufferProperties->Set(IDecoderOutputOptionNames::PTS, Electra::FVariantValue(Electra::FTimeValue(PresentationTimeTicks)));
 		BufferProperties->Set(IDecoderOutputOptionNames::Duration, Electra::FVariantValue(Electra::FTimeValue(FrameDurationTicks)));
 
+		// Extract timecode from bitstream processor SEI data
+		if (BitstreamProcessor.IsValid() && FrameBitstreamInfo.IsValid())
+		{
+			TMap<FString, FVariant> BsiProperties;
+			BitstreamProcessor->SetPropertiesOnOutput(BsiProperties, FrameBitstreamInfo);
+
+			TArray<uint8> PicTiming = ElectraDecodersUtil::GetVariantValueUInt8Array(BsiProperties, IElectraDecoderBitstreamProcessorInfo::CommonPictureTiming);
+			if (PicTiming.Num() == sizeof(ElectraDecodersUtil::MPEG::FCommonPictureTiming))
+			{
+				TSharedPtr<Electra::MPEG::FVideoDecoderTimecode> Timecode = MakeShared<Electra::MPEG::FVideoDecoderTimecode, ESPMode::ThreadSafe>();
+				Timecode->UpdateWith(*reinterpret_cast<const ElectraDecodersUtil::MPEG::FCommonPictureTiming*>(PicTiming.GetData()));
+				BufferProperties->Set(IDecoderOutputOptionNames::Timecode, Electra::FVariantValue(Timecode));
+			}
+		}
+
 		FString SetupOutputTextureSampleError;
 		const bool bSuccess = FElectraDecodersPlatformResources::SetupOutputTextureSample(
 			SetupOutputTextureSampleError, 
@@ -571,8 +597,20 @@ void FRtpDecoder::PollDecodedOutput()
 			continue;
 		}
 
-		TextureSample->SetTime(FMediaTimeStamp(CurrentPresentationTimestamp));
-
+		if (UE_LOG_ACTIVE(LogRtspMedia, VeryVerbose))
+		{
+			TOptional<FTimecode> OptionalTimecode = TextureSample->GetTimecode();
+			if (OptionalTimecode.IsSet())
+			{
+				const FTimecode Timecode = OptionalTimecode.GetValue();
+				UE_LOG(LogRtspMedia, VeryVerbose, TEXT("Adding video sample with timecode: %s with subframe: %f"), *Timecode.ToString(), Timecode.Subframe);
+			}
+			else
+			{
+				UE_LOG(LogRtspMedia, VeryVerbose, TEXT("Adding video sample without timecode"));
+			}
+		}
+		
 		Samples->AddVideo(TextureSample.ToSharedRef());
 	}
 }
diff --git a/Engine/Plugins/Media/RTSPMedia/Source/RTSPMedia/Private/RTP/RtpDecoder.h b/Engine/Plugins/Media/RTSPMedia/Source/RTSPMedia/Private/RTP/RtpDecoder.h
--- a/Engine/Plugins/Media/RTSPMedia/Source/RTSPMedia/Private/RTP/RtpDecoder.h
+++ b/Engine/Plugins/Media/RTSPMedia/Source/RTSPMedia/Private/RTP/RtpDecoder.h
@@ -16,6 +16,7 @@
 #include "Templates/SharedPointer.h"
 
 class IElectraDecoder;
+class IElectraDecoderBitstreamInfo;
 class IElectraDecoderBitstreamProcessor;
 class IElectraCodecFactory;
 class IElectraDecoderVideoOutput;
@@ -102,6 +103,7 @@ private:
 	// Decoding State
 	TArray<uint8> FrameBuffer;
 	TUniquePtr<FSdpSession> SdpSession;
+	TQueue<TSharedPtr<IElectraDecoderBitstreamInfo, ESPMode::ThreadSafe>> BitstreamInfoQueue;
 	uint32 VideoClockRate = 90000;
 	bool bWaitForIDRFrame = true;
 
