[RTSP Media] Handle mid-stream resolution changes via deferred decoder creation

#rnx
#Synced-CL 51630826

diff --git a/Engine/Plugins/Media/RTSPMedia/Source/RTSPMedia/Private/RTP/RtpDecoder.cpp b/Engine/Plugins/Media/RTSPMedia/Source/RTSPMedia/Private/RTP/RtpDecoder.cpp
--- a/Engine/Plugins/Media/RTSPMedia/Source/RTSPMedia/Private/RTP/RtpDecoder.cpp
+++ b/Engine/Plugins/Media/RTSPMedia/Source/RTSPMedia/Private/RTP/RtpDecoder.cpp
@@ -98,9 +98,6 @@ bool FRtpDecoder::Initialize(FMediaSamples* InSamples, const FSdpSession& InSess
 		return false;
 	}
 
-	SequenceParameterSet = VideoTrack->SequenceParameterSet;
-	PictureParameterSet = VideoTrack->PictureParameterSet;
-
 	// Store the clock rate to construct presentation time stamps when decoding
 	const uint32 ClockRate = VideoTrack->ClockRate;
 	if (ClockRate == 0)
@@ -109,17 +106,36 @@ bool FRtpDecoder::Initialize(FMediaSamples* InSamples, const FSdpSession& InSess
 		return false;
 	}
 	VideoClockRate = ClockRate;
-	
-	return CreateElectraDecoder(SequenceParameterSet, PictureParameterSet);
-}
 
-bool FRtpDecoder::Start()
-{
-	if (!ensureMsgf(ElectraDecoder.IsValid(), TEXT("Decoder cannot be started if electra decoder is not initialised")))
+	// Seed SPS/PPS from SDP if available. Some encoders only provide parameter sets
+	// out-of-band via sprop-parameter-sets and never send them in-band.
+	// In-band updates will override these via the existing comparison in HandleOnH264NalUnit.
+	const bool bHasSPS = !VideoTrack->SequenceParameterSet.IsEmpty();
+	const bool bHasPPS = !VideoTrack->PictureParameterSet.IsEmpty();
+	if (bHasSPS && bHasPPS)
 	{
-		return false;
+		SequenceParameterSet = VideoTrack->SequenceParameterSet;
+		PictureParameterSet = VideoTrack->PictureParameterSet;
+		bPendingCSDChange = true;
+		UE_LOG(LogRtspMedia, Verbose, TEXT("Seeded SPS (%d bytes) and PPS (%d bytes) from SDP"),
+			SequenceParameterSet.Num(), PictureParameterSet.Num());
+	}
+	else if (bHasSPS || bHasPPS)
+	{
+		UE_LOG(LogRtspMedia, Warning, TEXT("SDP contains %s but not %s, deferring decoder creation until in-band parameter sets arrive"),
+			bHasSPS ? TEXT("SPS") : TEXT("PPS"),
+			bHasSPS ? TEXT("PPS") : TEXT("SPS"));
+	}
+	else
+	{
+		UE_LOG(LogRtspMedia, Verbose, TEXT("No SPS or PPS in SDP, deferring decoder creation until in-band parameter sets arrive"));
 	}
 
+	return true;
+}
+
+bool FRtpDecoder::Start()
+{
 	bStopping.store(false);
 	
 	Thread.Reset(FRunnableThread::Create(this, TEXT("RtpDecoderThread"), 0, TPri_Normal));
@@ -332,12 +348,27 @@ void FRtpDecoder::Stop()
 	return true;
  }
 
-void FRtpDecoder::HandleOnH264NalUnit(TArray<uint8> InNalUnit, uint32 InTimestamp, bool bIsIDRFrame, bool bIsMarkerBitSet, bool bIsSlice)
+void FRtpDecoder::HandleOnH264NalUnit(FRtpH264NalUnit InNalUnit)
 {
+	// Track in-band SPS/PPS for mid-stream resolution changes (e.g. iPad orientation change).
+	if (!InNalUnit.SequenceParameterSet.IsEmpty() && InNalUnit.SequenceParameterSet != SequenceParameterSet)
+	{
+		UE_LOG(LogRtspMedia, Verbose, TEXT("In-band SPS changed (%d bytes)"), InNalUnit.SequenceParameterSet.Num());
+		SequenceParameterSet = MoveTemp(InNalUnit.SequenceParameterSet);
+		bPendingCSDChange = true;
+	}
+
+	if (!InNalUnit.PictureParameterSet.IsEmpty() && InNalUnit.PictureParameterSet != PictureParameterSet)
+	{
+		UE_LOG(LogRtspMedia, Verbose, TEXT("In-band PPS changed (%d bytes)"), InNalUnit.PictureParameterSet.Num());
+		PictureParameterSet = MoveTemp(InNalUnit.PictureParameterSet);
+		bPendingCSDChange = true;
+	}
+
 	// Wait on an IDR frame first then allow all other units through.
 	if (bWaitForIDRFrame)
 	{
-		if (!bIsIDRFrame)
+		if (!InNalUnit.bIsIDRFrame)
 		{
 			// Abort processing until we get one
 			return;
@@ -349,13 +380,13 @@ void FRtpDecoder::HandleOnH264NalUnit(TArray<uint8> InNalUnit, uint32 InTimestam
 	}
 
 	// Append the NAL unit to our frame buffer
-	FrameBuffer.Append(InNalUnit);
-	
+	FrameBuffer.Append(InNalUnit.Data);
+
 	// Marker may indicate a complete access unit.
 	// However, we can't use it to reliably start decoding.
 	// https://www.rfc-editor.org/rfc/rfc6184#section-5.1
-	// See the bIsSlice check below. 
-	if (!bIsMarkerBitSet)
+	// See the bIsSlice check below.
+	if (!InNalUnit.bIsMarkerBitSet)
 	{
 		return;
 	}
@@ -365,7 +396,7 @@ void FRtpDecoder::HandleOnH264NalUnit(TArray<uint8> InNalUnit, uint32 InTimestam
 	// So once the marker bit is set we need to check that we have received a NAL unit that represents an image slice.
 	// Any metadata units we receive will be aggregated into the next access unit.
 	// With the marker bit set and bIsSlice == true, we should have a complete unit containing image data ready for decoding.
-	if (!bIsSlice)
+	if (!InNalUnit.bIsSlice)
 	{
 		return;
 	}
@@ -390,21 +421,44 @@ void FRtpDecoder::HandleOnH264NalUnit(TArray<uint8> InNalUnit, uint32 InTimestam
 	IElectraDecoder::FInputAccessUnit AccessUnit;
 	AccessUnit.Data = FrameBuffer.GetData();
 	AccessUnit.DataSize = FrameBuffer.Num();
-	AccessUnit.PTS = FTimespan::FromSeconds(static_cast<double>(InTimestamp) / VideoClockRate);
-	if (bIsIDRFrame)
+	AccessUnit.PTS = FTimespan::FromSeconds(static_cast<double>(InNalUnit.Timestamp) / VideoClockRate);
+	if (InNalUnit.bIsIDRFrame)
 	{
 		AccessUnit.Flags = EElectraDecoderFlags::IsSyncSample;
 	}
 
-	if (!BitstreamProcessor.IsValid())
+	// Create or recreate the decoder when SPS/PPS has changed.
+	// On the first frame this creates the decoder from in-band SPS/PPS (deferred from Initialize).
+	// On subsequent changes (e.g. iPad orientation) it recreates with the new parameters.
+	if (bPendingCSDChange)
 	{
-		UE_LOG(LogRtspMedia, Warning, TEXT("BitstreamProcessor unavailable prior to decoding"));
+		if (SequenceParameterSet.IsEmpty() || PictureParameterSet.IsEmpty())
+		{
+			UE_LOG(LogRtspMedia, Verbose, TEXT("Waiting for both SPS and PPS before creating decoder"));
+			return;
+		}
+
+		UE_LOG(LogRtspMedia, Verbose, TEXT("CSD changed, %s decoder"), ElectraDecoder.IsValid() ? TEXT("recreating") : TEXT("creating"));
+		bPendingCSDChange = false;
+
+		if (!CreateElectraDecoder(SequenceParameterSet, PictureParameterSet))
+		{
+			UE_LOG(LogRtspMedia, Error, TEXT("Failed to create decoder from in-band SPS/PPS"));
+			OnDecoderError.ExecuteIfBound();
+			bStopping.store(true);
+			return;
+		}
+	}
+
+	if (!BitstreamProcessor.IsValid() || !ElectraDecoder.IsValid())
+	{
+		UE_LOG(LogRtspMedia, Warning, TEXT("Decoder not yet available, waiting for in-band SPS/PPS"));
 		return;
 	}
 
 	// Before decoding can take place, any platform specific preparation needs to be performed.
 	// On some platforms the big endian 4 byte length indicator we output in the depacketizer will be swapped
-	// for a H264 Annex B start code
+	// for a H264 Annex B start code.
 	// If we were using ElectraPlayer directly then we wouldn't need to perform this step.
 	// However, as we're directly integrating with the underlying decoders we need to prepare the data for processing.
 	TSharedPtr<IElectraDecoderBitstreamInfo, ESPMode::ThreadSafe> BitstreamInfo;
@@ -415,12 +469,6 @@ void FRtpDecoder::HandleOnH264NalUnit(TArray<uint8> InNalUnit, uint32 InTimestam
 		return;
 	}
 
-	if (!ElectraDecoder.IsValid())
-	{
-		UE_LOG(LogRtspMedia, Warning, TEXT("ElectraDecoder unavailable prior to decoding"));
-		return;
-	}
-
 	const IElectraDecoder::EDecoderError DecodeResult = ElectraDecoder->DecodeAccessUnit(AccessUnit, CodecOptions);
 
 	switch (DecodeResult)
@@ -437,10 +485,7 @@ void FRtpDecoder::HandleOnH264NalUnit(TArray<uint8> InNalUnit, uint32 InTimestam
 			if (!CreateElectraDecoder(SequenceParameterSet, PictureParameterSet))
 			{
 				UE_LOG(LogRtspMedia, Error, TEXT("Failed to re-create ElectraDecoder after receiving a LostDecoder error"));
-				if (OnDecoderError.IsBound())
-				{
-					OnDecoderError.Execute();
-				}
+				OnDecoderError.ExecuteIfBound();
 				bStopping.store(true);
 			}
 			return;
@@ -458,7 +503,7 @@ void FRtpDecoder::PollDecodedOutput()
 {
 	if (!ElectraDecoder.IsValid())
 	{
-		UE_LOG(LogRtspMedia, Warning, TEXT("Attempted to process decoded output without an electra decoder instance"));
+		UE_LOG(LogRtspMedia, VeryVerbose, TEXT("Attempted to process decoded output without an electra decoder instance"));
 		return;
 	}
 
diff --git a/Engine/Plugins/Media/RTSPMedia/Source/RTSPMedia/Private/RTP/RtpDecoder.h b/Engine/Plugins/Media/RTSPMedia/Source/RTSPMedia/Private/RTP/RtpDecoder.h
--- a/Engine/Plugins/Media/RTSPMedia/Source/RTSPMedia/Private/RTP/RtpDecoder.h
+++ b/Engine/Plugins/Media/RTSPMedia/Source/RTSPMedia/Private/RTP/RtpDecoder.h
@@ -2,6 +2,8 @@
 
 #pragma once
 
+#include "RtpH264NalUnit.h"
+
 #include "RtspMediaDefaults.h"
 
 #include "RtpPacket.h"
@@ -73,7 +75,7 @@ protected:
 private:
 
 	bool CreateElectraDecoder(const TArray<uint8>& InSequenceParameterSet, const TArray<uint8>& InPictureParameterSet);
-	void HandleOnH264NalUnit(TArray<uint8> InNalUnit, uint32 InTimestamp, bool bIsIDRFrame, bool bIsMarkerBitSet, bool bIsSlice);
+	void HandleOnH264NalUnit(FRtpH264NalUnit InNalUnit);
 	void PollDecodedOutput();
 	
 	void ResetFrameBuffer();
@@ -106,6 +108,7 @@ private:
 	TQueue<TSharedPtr<IElectraDecoderBitstreamInfo, ESPMode::ThreadSafe>> BitstreamInfoQueue;
 	uint32 VideoClockRate = 90000;
 	bool bWaitForIDRFrame = true;
+	bool bPendingCSDChange = false;
 
 	// Texture Sample Generation
 	int32 CurrentVideoWidth = 0;
diff --git a/Engine/Plugins/Media/RTSPMedia/Source/RTSPMedia/Private/RTP/RtpH264Depacketizer.cpp b/Engine/Plugins/Media/RTSPMedia/Source/RTSPMedia/Private/RTP/RtpH264Depacketizer.cpp
--- a/Engine/Plugins/Media/RTSPMedia/Source/RTSPMedia/Private/RTP/RtpH264Depacketizer.cpp
+++ b/Engine/Plugins/Media/RTSPMedia/Source/RTSPMedia/Private/RTP/RtpH264Depacketizer.cpp
@@ -229,22 +229,12 @@ void FRtpH264Depacketizer::EmitNalUnit(TArrayView<const uint8> InNalUnit, uint32
 		return;
 	}
 
-	// Determine whether this NAL unit represents an IDR-Frame
-	// This is a special kind of I-Frame that allows the decoder to clear
-	// its internal state. The last 5 bits of the first byte in the unit
-	// contain the NAL type. 5 == an IDR frame.
+	// Determine the NAL type from the first byte (last 5 bits).
 	// 0x1F = 00011111
 	const uint8 NalType = InNalUnit[0] & 0x1F;
-	const bool bIsIdrFrame = NalType == 5;
-
-	// Because certain senders will set the marker bit to 1 even if we haven't
-	// received image data for that timestamp we need to also track if the NAL type is an image 'slice'.
-	// I.e. Not a NAL unit containing only metadata.
-	const bool bIsSlice = NalType >= 1 && NalType <= 5;
 
 	// The final NAL unit output needs to be in the AVCC format with the NAL unit data itself
 	// prefixed with a 4-byte big endian value describing the length of the contained NAL unit.
-
 	const int32 NalUnitSize = InNalUnit.Num();
 	const uint8 HeaderByte0 = (NalUnitSize >> 24) & 0xFF;
 	const uint8 HeaderByte1 = (NalUnitSize >> 16) & 0xFF;
@@ -257,5 +247,23 @@ void FRtpH264Depacketizer::EmitNalUnit(TArrayView<const uint8> InNalUnit, uint32
 	AvccNalUnit.Add(HeaderByte2);
 	AvccNalUnit.Add(HeaderByte3);
 	AvccNalUnit.Append(InNalUnit.GetData(), NalUnitSize);
-	OnH264NalUnit.ExecuteIfBound(MoveTemp(AvccNalUnit), InTimestamp, bIsIdrFrame, bInMarker, bIsSlice);
+
+	FRtpH264NalUnit NalUnit;
+	NalUnit.Data = MoveTemp(AvccNalUnit);
+	NalUnit.Timestamp = InTimestamp;
+	NalUnit.NalType = NalType;
+	NalUnit.bIsIDRFrame = NalType == 5;
+	NalUnit.bIsMarkerBitSet = bInMarker;
+	NalUnit.bIsSlice = NalType >= 1 && NalType <= 5;
+
+	if (NalType == 7)
+	{
+		NalUnit.SequenceParameterSet = TArray<uint8>(InNalUnit.GetData(), InNalUnit.Num());
+	}
+	else if (NalType == 8)
+	{
+		NalUnit.PictureParameterSet = TArray<uint8>(InNalUnit.GetData(), InNalUnit.Num());
+	}
+
+	OnH264NalUnit.ExecuteIfBound(MoveTemp(NalUnit));
 }
diff --git a/Engine/Plugins/Media/RTSPMedia/Source/RTSPMedia/Private/RTP/RtpH264Depacketizer.h b/Engine/Plugins/Media/RTSPMedia/Source/RTSPMedia/Private/RTP/RtpH264Depacketizer.h
--- a/Engine/Plugins/Media/RTSPMedia/Source/RTSPMedia/Private/RTP/RtpH264Depacketizer.h
+++ b/Engine/Plugins/Media/RTSPMedia/Source/RTSPMedia/Private/RTP/RtpH264Depacketizer.h
@@ -2,12 +2,11 @@
 
 #pragma once
 
-#include "RtspMediaDefaults.h"
+#include "RtpH264NalUnit.h"
 
-#include "CoreMinimal.h"
+#include "RtspMediaDefaults.h"
 
-// NAL Unit, Timestamp, bIsIdrFrame, bIsMarkerBitSet, bIsSlice 
-DECLARE_DELEGATE_FiveParams(FOnH264NalUnit, TArray<uint8>, uint32, bool, bool, bool);
+DECLARE_DELEGATE_OneParam(FOnH264NalUnit, FRtpH264NalUnit);
 
 class FRtpH264Depacketizer
 {
diff --git a/Engine/Plugins/Media/RTSPMedia/Source/RTSPMedia/Private/RTP/RtpH264NalUnit.h b/Engine/Plugins/Media/RTSPMedia/Source/RTSPMedia/Private/RTP/RtpH264NalUnit.h
new file mode 100644
--- /dev/null
+++ b/Engine/Plugins/Media/RTSPMedia/Source/RTSPMedia/Private/RTP/RtpH264NalUnit.h
@@ -0,0 +1,17 @@
+// Copyright Epic Games, Inc. All Rights Reserved.
+
+#pragma once
+
+#include "CoreMinimal.h"
+
+struct FRtpH264NalUnit
+{
+	TArray<uint8> Data;                    // AVCC-formatted NAL unit (for frame buffer)
+	TArray<uint8> SequenceParameterSet;    // Raw SPS bytes (non-empty only for NAL type 7)
+	TArray<uint8> PictureParameterSet;     // Raw PPS bytes (non-empty only for NAL type 8)
+	uint32 Timestamp = 0;
+	uint8 NalType = 0;
+	bool bIsIDRFrame = false;
+	bool bIsMarkerBitSet = false;
+	bool bIsSlice = false;
+};
