[RTSPMedia] Add Provide CPU buffer option

#rnx

#Synced-CL 52001901

diff --git a/Engine/Plugins/Media/RTSPMedia/Source/RTSPMedia/Private/Player/RtspMediaPlayer.cpp b/Engine/Plugins/Media/RTSPMedia/Source/RTSPMedia/Private/Player/RtspMediaPlayer.cpp
--- a/Engine/Plugins/Media/RTSPMedia/Source/RTSPMedia/Private/Player/RtspMediaPlayer.cpp
+++ b/Engine/Plugins/Media/RTSPMedia/Source/RTSPMedia/Private/Player/RtspMediaPlayer.cpp
@@ -96,6 +96,7 @@ bool FRtspMediaPlayer::Open(const FString& InUrl, const IMediaOptions* InOptions
 	MaxFragmentBufferSizeMbOption = InOptions->GetMediaOption(RtspMedia::Option::MaxFragmentBufferSize, static_cast<int64>(RtspMedia::Default::MaxFragmentBufferSizeMb));
 	DecoderBufferSizeOption = InOptions->GetMediaOption(RtspMedia::Option::DecoderBufferSize, static_cast<int64>(RtspMedia::Default::DecoderBufferSize));
 	DecoderPollIntervalMsOption = InOptions->GetMediaOption(RtspMedia::Option::DecoderPollIntervalMs, static_cast<int64>(RtspMedia::Default::DecoderPollIntervalMs));
+	bProvideCpuBufferOption = InOptions->GetMediaOption(RtspMedia::Option::ProvideCpuBuffer, RtspMedia::Default::bProvideCpuBuffer);
 
 	ResetClient();
 	ResetDecoder();
@@ -247,7 +248,8 @@ void FRtspMediaPlayer::HandleOnSdpSessionReceived(const FSdpSession& Session)
 		Session,
 		MaxFragmentBufferSize,
 		DecoderBufferSize,
-		DecoderPollIntervalMs))
+		DecoderPollIntervalMs,
+		bProvideCpuBufferOption))
 	{
 		UE_LOG(LogRtspMedia, Error, TEXT("Failed to initialise decoder"));
 		SetState(EMediaState::Error);
diff --git a/Engine/Plugins/Media/RTSPMedia/Source/RTSPMedia/Private/Player/RtspMediaPlayer.h b/Engine/Plugins/Media/RTSPMedia/Source/RTSPMedia/Private/Player/RtspMediaPlayer.h
--- a/Engine/Plugins/Media/RTSPMedia/Source/RTSPMedia/Private/Player/RtspMediaPlayer.h
+++ b/Engine/Plugins/Media/RTSPMedia/Source/RTSPMedia/Private/Player/RtspMediaPlayer.h
@@ -61,6 +61,7 @@ private:
 	int64 MaxFragmentBufferSizeMbOption = RtspMedia::Default::MaxFragmentBufferSizeMb;
 	int64 DecoderBufferSizeOption = RtspMedia::Default::DecoderBufferSize;
 	int64 DecoderPollIntervalMsOption = RtspMedia::Default::DecoderPollIntervalMs;
+	bool bProvideCpuBufferOption = RtspMedia::Default::bProvideCpuBuffer;
 	TUniquePtr<FRtpDecoder> RtpDecoder;
 	
 	// Media Player
diff --git a/Engine/Plugins/Media/RTSPMedia/Source/RTSPMedia/Private/RTP/RtpDecoder.cpp b/Engine/Plugins/Media/RTSPMedia/Source/RTSPMedia/Private/RTP/RtpDecoder.cpp
--- a/Engine/Plugins/Media/RTSPMedia/Source/RTSPMedia/Private/RTP/RtpDecoder.cpp
+++ b/Engine/Plugins/Media/RTSPMedia/Source/RTSPMedia/Private/RTP/RtpDecoder.cpp
@@ -21,6 +21,8 @@
 #include "Utils/MPEG/ElectraUtilsMPEGVideo.h"
 #include "Utils/VideoDecoderHelpers.h"
 
+#include "RtpDecoderCpuBuffer.h"
+
 #include "Async/Async.h"
 #include "HAL/Event.h"
 #include "HAL/PlatformProcess.h"
@@ -69,12 +71,13 @@ FRtpDecoder::~FRtpDecoder()
 	}
 }
 
-bool FRtpDecoder::Initialize(FMediaSamples* InSamples, const FSdpSession& InSession, int64 InMaxFragmentBufferSize, int64 InDecoderBufferSize, uint32 InDecoderPollIntervalMs)
+bool FRtpDecoder::Initialize(FMediaSamples* InSamples, const FSdpSession& InSession, int64 InMaxFragmentBufferSize, int64 InDecoderBufferSize, uint32 InDecoderPollIntervalMs, bool bInProvideCpuBuffer)
 {
 	Samples = InSamples;
 	MaxFragmentBufferSize = InMaxFragmentBufferSize;
 	DecoderBufferSize = InDecoderBufferSize;
 	DecoderPollIntervalMs = InDecoderPollIntervalMs;
+	bProvideCpuBuffer = bInProvideCpuBuffer;
 	TextureSamplePool = MakeUnique<FElectraTextureSamplePool>();
 	
 	SdpSession = MakeUnique<FSdpSession>(InSession);
@@ -241,6 +244,10 @@ void FRtpDecoder::Stop()
 		CodecOptions.Emplace("dcr", AvcConfigurationRecord.GetRawData());
 		CodecOptions.Emplace("csd", AvcConfigurationRecord.GetCodecSpecificData());
 		CodecOptions.Emplace("max_output_buffers", FVariant(DecoderBufferSize));
+		if (bProvideCpuBuffer)
+		{
+			CodecOptions.Emplace("force_cpu_output", FVariant(static_cast<int64>(1)));
+		}
 
 		// 'avc1' is a flavour of H.264 where the Sequence and Picture parameter sets are stored 'out of band'
 		// This means they are stored separately from the main H.264 bitstream, which is the case when streaming with RTSP.
@@ -642,6 +649,17 @@ void FRtpDecoder::PollDecodedOutput()
 			continue;
 		}
 
+		if (bProvideCpuBuffer)
+		{
+			RtspMedia::PopulateCpuBuffer(TextureSample.Get(), VideoOutput);
+			if (!TextureSample->Buffer.IsValid() && !bLoggedCpuBufferWarning)
+			{
+				UE_LOG(LogRtspMedia, Warning, TEXT("bProvideCpuBuffer is enabled but the active decoder did not provide CPU-accessible data. "
+					"GetBuffer() will return nullptr. Ensure the selected H.264 decoder supports the force_cpu_output codec option."));
+				bLoggedCpuBufferWarning = true;
+			}
+		}
+
 		if (UE_LOG_ACTIVE(LogRtspMedia, VeryVerbose))
 		{
 			TOptional<FTimecode> OptionalTimecode = TextureSample->GetTimecode();
diff --git a/Engine/Plugins/Media/RTSPMedia/Source/RTSPMedia/Private/RTP/RtpDecoder.h b/Engine/Plugins/Media/RTSPMedia/Source/RTSPMedia/Private/RTP/RtpDecoder.h
--- a/Engine/Plugins/Media/RTSPMedia/Source/RTSPMedia/Private/RTP/RtpDecoder.h
+++ b/Engine/Plugins/Media/RTSPMedia/Source/RTSPMedia/Private/RTP/RtpDecoder.h
@@ -54,7 +54,8 @@ public:
 		const FSdpSession& InSession,
 		int64 InMaxFragmentBufferSize,
 		int64 InDecoderBufferSize,
-		uint32 InDecoderPollIntervalMs);
+		uint32 InDecoderPollIntervalMs,
+		bool bInProvideCpuBuffer);
 	
 	bool Start();
 	void Shutdown();
@@ -77,7 +78,7 @@ private:
 	bool CreateElectraDecoder(const TArray<uint8>& InSequenceParameterSet, const TArray<uint8>& InPictureParameterSet);
 	void HandleOnH264NalUnit(FRtpH264NalUnit InNalUnit);
 	void PollDecodedOutput();
-	
+
 	void ResetFrameBuffer();
 	void ResetPlatformVideoResource();
 
@@ -101,6 +102,8 @@ private:
 	int64 MaxFragmentBufferSize = RtspMedia::Default::MaxFragmentBufferSizeBytes;
 	int64 DecoderBufferSize = RtspMedia::Default::DecoderBufferSize;
 	uint32 DecoderPollIntervalMs = RtspMedia::Default::DecoderPollIntervalMs;
+	bool bProvideCpuBuffer = RtspMedia::Default::bProvideCpuBuffer;
+	bool bLoggedCpuBufferWarning = false;
 
 	// Decoding State
 	TArray<uint8> FrameBuffer;
diff --git a/Engine/Plugins/Media/RTSPMedia/Source/RTSPMedia/Private/RTP/RtpDecoderCpuBuffer.h b/Engine/Plugins/Media/RTSPMedia/Source/RTSPMedia/Private/RTP/RtpDecoderCpuBuffer.h
new file mode 100644
--- /dev/null
+++ b/Engine/Plugins/Media/RTSPMedia/Source/RTSPMedia/Private/RTP/RtpDecoderCpuBuffer.h
@@ -0,0 +1,22 @@
+// Copyright Epic Games, Inc. All Rights Reserved.
+
+#pragma once
+
+#include "Templates/SharedPointer.h"
+
+class FElectraTextureSample;
+class IElectraDecoderVideoOutput;
+
+namespace RtspMedia
+{
+	/**
+	 * Populates FElectraTextureSample::Buffer with NV12 pixel data for CPU access.
+	 * Called after SetupOutputTextureSample succeeds. The GPU rendering path is unaffected.
+	 *
+	 * Platform behavior:
+	 * - macOS/iOS: Locks the CVPixelBuffer and copies Y + CbCr planes
+	 * - Windows: Extracts MFSample via GetPlatformOutputHandle and copies via Lock2D
+	 * - Linux: No-op (SetupOutputTextureSample already populates Buffer)
+	 */
+	void PopulateCpuBuffer(FElectraTextureSample* InTextureSample, const TSharedPtr<IElectraDecoderVideoOutput>& InVideoOutput);
+}
diff --git a/Engine/Plugins/Media/RTSPMedia/Source/RTSPMedia/Private/RTP/RtpDecoderCpuBuffer_Apple.cpp b/Engine/Plugins/Media/RTSPMedia/Source/RTSPMedia/Private/RTP/RtpDecoderCpuBuffer_Apple.cpp
new file mode 100644
--- /dev/null
+++ b/Engine/Plugins/Media/RTSPMedia/Source/RTSPMedia/Private/RTP/RtpDecoderCpuBuffer_Apple.cpp
@@ -0,0 +1,80 @@
+// Copyright Epic Games, Inc. All Rights Reserved.
+
+#include "HAL/Platform.h"
+
+#if PLATFORM_APPLE
+
+#include "RtpDecoderCpuBuffer.h"
+
+#include "RtspMediaConstants.h"
+
+#include "ElectraTextureSample.h"
+
+#include "Misc/ScopeExit.h"
+
+#include <CoreVideo/CoreVideo.h>
+
+void RtspMedia::PopulateCpuBuffer(FElectraTextureSample* InTextureSample, const TSharedPtr<IElectraDecoderVideoOutput>& InVideoOutput)
+{
+	if (!InTextureSample || !InVideoOutput.IsValid())
+	{
+		return;
+	}
+
+	// On Apple platforms, SetupOutputTextureSample populates ImageBufferRef for GPU rendering.
+	// Lock the CVPixelBuffer to copy its contents into Buffer for CPU access.
+	CVPixelBufferRef PixelBuffer = InTextureSample->ImageBufferRef;
+	if (!PixelBuffer)
+	{
+		UE_LOG(LogRtspMedia, Warning, TEXT("PopulateCpuBuffer: No CVPixelBuffer available"));
+		return;
+	}
+
+	const CVReturn LockResult = CVPixelBufferLockBaseAddress(PixelBuffer, kCVPixelBufferLock_ReadOnly);
+	if (LockResult != kCVReturnSuccess)
+	{
+		UE_LOG(LogRtspMedia, Warning, TEXT("PopulateCpuBuffer: CVPixelBufferLockBaseAddress failed (%d)"), static_cast<int32>(LockResult));
+		return;
+	}
+
+	ON_SCOPE_EXIT
+	{
+		CVPixelBufferUnlockBaseAddress(PixelBuffer, kCVPixelBufferLock_ReadOnly);
+	};
+
+	const size_t PlaneCount = CVPixelBufferGetPlaneCount(PixelBuffer);
+	if (PlaneCount < 2)
+	{
+		UE_LOG(LogRtspMedia, Warning, TEXT("PopulateCpuBuffer: Expected 2 planes (NV12), got %d"), static_cast<int32>(PlaneCount));
+		return;
+	}
+
+	// Y plane
+	const uint8* YBaseAddress = static_cast<const uint8*>(CVPixelBufferGetBaseAddressOfPlane(PixelBuffer, 0));
+	const size_t YBytesPerRow = CVPixelBufferGetBytesPerRowOfPlane(PixelBuffer, 0);
+	const size_t YHeight = CVPixelBufferGetHeightOfPlane(PixelBuffer, 0);
+
+	// CbCr plane (interleaved UV)
+	const uint8* CbCrBaseAddress = static_cast<const uint8*>(CVPixelBufferGetBaseAddressOfPlane(PixelBuffer, 1));
+	const size_t CbCrBytesPerRow = CVPixelBufferGetBytesPerRowOfPlane(PixelBuffer, 1);
+	const size_t CbCrHeight = CVPixelBufferGetHeightOfPlane(PixelBuffer, 1);
+
+	if (!YBaseAddress || !CbCrBaseAddress)
+	{
+		UE_LOG(LogRtspMedia, Warning, TEXT("PopulateCpuBuffer: CVPixelBufferGetBaseAddressOfPlane returned null (Y=%p, CbCr=%p)"), YBaseAddress, CbCrBaseAddress);
+		return;
+	}
+
+	const size_t TotalBytes = (YBytesPerRow * YHeight) + (CbCrBytesPerRow * CbCrHeight);
+
+	InTextureSample->Buffer = MakeShared<TArray64<uint8>, ESPMode::ThreadSafe>();
+	InTextureSample->Buffer->SetNumUninitialized(static_cast<int64>(TotalBytes));
+
+	uint8* Dest = InTextureSample->Buffer->GetData();
+	FMemory::Memcpy(Dest, YBaseAddress, YBytesPerRow * YHeight);
+	FMemory::Memcpy(Dest + (YBytesPerRow * YHeight), CbCrBaseAddress, CbCrBytesPerRow * CbCrHeight);
+
+	InTextureSample->Stride = static_cast<uint32>(YBytesPerRow);
+}
+
+#endif // PLATFORM_APPLE
diff --git a/Engine/Plugins/Media/RTSPMedia/Source/RTSPMedia/Private/RTP/RtpDecoderCpuBuffer_Linux.cpp b/Engine/Plugins/Media/RTSPMedia/Source/RTSPMedia/Private/RTP/RtpDecoderCpuBuffer_Linux.cpp
new file mode 100644
--- /dev/null
+++ b/Engine/Plugins/Media/RTSPMedia/Source/RTSPMedia/Private/RTP/RtpDecoderCpuBuffer_Linux.cpp
@@ -0,0 +1,14 @@
+// Copyright Epic Games, Inc. All Rights Reserved.
+
+#include "HAL/Platform.h"
+
+#if PLATFORM_LINUX
+
+#include "RtpDecoderCpuBuffer.h"
+
+void RtspMedia::PopulateCpuBuffer(FElectraTextureSample* InTextureSample, const TSharedPtr<IElectraDecoderVideoOutput>& InVideoOutput)
+{
+	// Linux already populates Buffer via SetupOutputTextureSample, nothing to do
+}
+
+#endif // PLATFORM_LINUX
diff --git a/Engine/Plugins/Media/RTSPMedia/Source/RTSPMedia/Private/RTP/RtpDecoderCpuBuffer_Windows.cpp b/Engine/Plugins/Media/RTSPMedia/Source/RTSPMedia/Private/RTP/RtpDecoderCpuBuffer_Windows.cpp
new file mode 100644
--- /dev/null
+++ b/Engine/Plugins/Media/RTSPMedia/Source/RTSPMedia/Private/RTP/RtpDecoderCpuBuffer_Windows.cpp
@@ -0,0 +1,97 @@
+// Copyright Epic Games, Inc. All Rights Reserved.
+
+#include "HAL/Platform.h"
+
+#if PLATFORM_WINDOWS
+
+#include "RtpDecoderCpuBuffer.h"
+
+#include "RtspMediaConstants.h"
+
+#include "ElectraTextureSample.h"
+#include "IElectraDecoderOutputVideo.h"
+#include "IElectraDecoderFeaturesAndOptions.h"
+
+#include "Misc/ScopeExit.h"
+
+#include "Windows/AllowWindowsPlatformTypes.h"
+#include <mfapi.h>
+#include <mfidl.h>
+#include <mfobjects.h>
+#include "Windows/HideWindowsPlatformTypes.h"
+
+void RtspMedia::PopulateCpuBuffer(FElectraTextureSample* InTextureSample, const TSharedPtr<IElectraDecoderVideoOutput>& InVideoOutput)
+{
+	if (!InTextureSample || !InVideoOutput.IsValid())
+	{
+		return;
+	}
+
+	// On Windows with DX12 GPU decode, SetupOutputTextureSample populates the GPU texture.
+	// The MFSample is also available on the output when force_cpu_output is set.
+	if (InTextureSample->Buffer.IsValid())
+	{
+		return;
+	}
+
+	void* MFSampleHandle = InVideoOutput->GetPlatformOutputHandle(EElectraDecoderPlatformOutputHandleType::MFSample);
+	if (!MFSampleHandle)
+	{
+		UE_LOG(LogRtspMedia, VeryVerbose, TEXT("PopulateCpuBuffer: No MFSample available on video output"));
+		return;
+	}
+
+	TRefCountPtr<IMFSample> Sample = static_cast<IMFSample*>(MFSampleHandle);
+	TRefCountPtr<IMFMediaBuffer> MediaBuffer;
+	HRESULT Result = Sample->GetBufferByIndex(0, MediaBuffer.GetInitReference());
+	if (FAILED(Result))
+	{
+		UE_LOG(LogRtspMedia, Warning, TEXT("PopulateCpuBuffer: GetBufferByIndex failed (0x%08X)"), static_cast<uint32>(Result));
+		return;
+	}
+
+	TRefCountPtr<IMF2DBuffer> Buffer2D;
+	Result = MediaBuffer->QueryInterface(__uuidof(IMF2DBuffer), reinterpret_cast<void**>(Buffer2D.GetInitReference()));
+	if (FAILED(Result))
+	{
+		UE_LOG(LogRtspMedia, Warning, TEXT("PopulateCpuBuffer: IMF2DBuffer query failed (0x%08X)"), static_cast<uint32>(Result));
+		return;
+	}
+
+	uint8* Data = nullptr;
+	LONG Pitch = 0;
+	Result = Buffer2D->Lock2D(&Data, &Pitch);
+	if (FAILED(Result))
+	{
+		UE_LOG(LogRtspMedia, Warning, TEXT("PopulateCpuBuffer: Lock2D failed (0x%08X)"), static_cast<uint32>(Result));
+		return;
+	}
+
+	ON_SCOPE_EXIT
+	{
+		Buffer2D->Unlock2D();
+	};
+
+	if (!Data || Pitch <= 0)
+	{
+		UE_LOG(LogRtspMedia, Warning, TEXT("PopulateCpuBuffer: Lock2D returned invalid data (Data=%p, Pitch=%d)"), Data, static_cast<int32>(Pitch));
+		return;
+	}
+
+	// NV12: Height * 3/2 to include both Y and UV planes
+	const int32 DecodedHeight = InVideoOutput->GetDecodedHeight();
+	const int64 TotalBytes = static_cast<int64>(Pitch) * DecodedHeight * 3 / 2;
+
+	if (TotalBytes <= 0)
+	{
+		UE_LOG(LogRtspMedia, Warning, TEXT("PopulateCpuBuffer: Invalid buffer dimensions (Pitch=%d, DecodedHeight=%d)"), static_cast<int32>(Pitch), DecodedHeight);
+		return;
+	}
+
+	InTextureSample->Buffer = MakeShared<TArray64<uint8>, ESPMode::ThreadSafe>();
+	InTextureSample->Buffer->SetNumUninitialized(TotalBytes);
+	FMemory::Memcpy(InTextureSample->Buffer->GetData(), Data, TotalBytes);
+	InTextureSample->Stride = static_cast<uint32>(Pitch);
+}
+
+#endif // PLATFORM_WINDOWS
diff --git a/Engine/Plugins/Media/RTSPMedia/Source/RTSPMedia/Private/RtspMediaConstants.h b/Engine/Plugins/Media/RTSPMedia/Source/RTSPMedia/Private/RtspMediaConstants.h
--- a/Engine/Plugins/Media/RTSPMedia/Source/RTSPMedia/Private/RtspMediaConstants.h
+++ b/Engine/Plugins/Media/RTSPMedia/Source/RTSPMedia/Private/RtspMediaConstants.h
@@ -16,6 +16,7 @@ namespace RtspMedia
 	namespace Option
 	{
 		inline const FName AutoReconnect("AutoReconnect");
+		inline const FName ProvideCpuBuffer("ProvideCpuBuffer");
 		inline const FName DecoderBufferSize("DecoderBufferSize");
 		inline const FName DecoderPollIntervalMs("DecoderPollIntervalMs");
 		inline const FName JitterBufferAutoAdjust("JitterBufferAutoAdjust");
diff --git a/Engine/Plugins/Media/RTSPMedia/Source/RTSPMedia/Private/RtspMediaSource.cpp b/Engine/Plugins/Media/RTSPMedia/Source/RTSPMedia/Private/RtspMediaSource.cpp
--- a/Engine/Plugins/Media/RTSPMedia/Source/RTSPMedia/Private/RtspMediaSource.cpp
+++ b/Engine/Plugins/Media/RTSPMedia/Source/RTSPMedia/Private/RtspMediaSource.cpp
@@ -40,6 +40,11 @@ bool URtspMediaSource::GetMediaOption(const FName& InKey, bool InDefaultValue) c
 		return bAutoReconnect;
 	}
 	
+	if (InKey == RtspMedia::Option::ProvideCpuBuffer)
+	{
+		return bProvideCpuBuffer;
+	}
+
 	if (InKey == RtspMedia::Option::JitterBufferAutoAdjust)
 	{
 		return bJitterBufferAutoAdjust;
@@ -125,7 +130,8 @@ double URtspMediaSource::GetMediaOption(const FName& InKey, double InDefaultValu
 
 bool URtspMediaSource::HasMediaOption(const FName& InKey) const
 {
-	if (InKey == RtspMedia::Option::AutoReconnect ||
+	if (InKey == RtspMedia::Option::ProvideCpuBuffer ||
+		InKey == RtspMedia::Option::AutoReconnect ||
 		InKey == RtspMedia::Option::DecoderBufferSize ||
 		InKey == RtspMedia::Option::DecoderPollIntervalMs ||
 		InKey == RtspMedia::Option::JitterBufferAutoAdjust ||
diff --git a/Engine/Plugins/Media/RTSPMedia/Source/RTSPMedia/Public/RtspMediaDefaults.h b/Engine/Plugins/Media/RTSPMedia/Source/RTSPMedia/Public/RtspMediaDefaults.h
--- a/Engine/Plugins/Media/RTSPMedia/Source/RTSPMedia/Public/RtspMediaDefaults.h
+++ b/Engine/Plugins/Media/RTSPMedia/Source/RTSPMedia/Public/RtspMediaDefaults.h
@@ -7,6 +7,7 @@
 namespace RtspMedia::Default
 {
 	inline constexpr bool bAutoReconnect = true;
+	inline constexpr bool bProvideCpuBuffer = false;
 	inline constexpr int32 DecoderBufferSize = 5;
 	inline constexpr int32 DecoderPollIntervalMs = 5;
 	inline constexpr bool bJitterBufferAutoAdjust = true;
diff --git a/Engine/Plugins/Media/RTSPMedia/Source/RTSPMedia/Public/RtspMediaSource.h b/Engine/Plugins/Media/RTSPMedia/Source/RTSPMedia/Public/RtspMediaSource.h
--- a/Engine/Plugins/Media/RTSPMedia/Source/RTSPMedia/Public/RtspMediaSource.h
+++ b/Engine/Plugins/Media/RTSPMedia/Source/RTSPMedia/Public/RtspMediaSource.h
@@ -109,6 +109,11 @@ public:
 	UPROPERTY(EditAnywhere, BlueprintReadWrite, Category="Decoder", meta=(ClampMin="2", ClampMax="16"))
 	int32 MaxQueuedVideoSamples = RtspMedia::Default::MaxQueuedVideoSamples;
 
+	/** Populate a CPU-accessible pixel buffer alongside the GPU texture
+	 * Whether this is available depends on the underlying decoder implementation */
+	UPROPERTY(EditAnywhere, BlueprintReadWrite, Category="Decoder", meta=(DisplayName="Provide CPU Buffer"))
+	bool bProvideCpuBuffer = RtspMedia::Default::bProvideCpuBuffer;
+
 	//~ Begin UMediaSource
 	virtual FString GetUrl() const override;
 	virtual bool Validate() const override;
diff --git a/Engine/Plugins/Media/RTSPMedia/Source/RTSPMedia/RTSPMedia.Build.cs b/Engine/Plugins/Media/RTSPMedia/Source/RTSPMedia/RTSPMedia.Build.cs
--- a/Engine/Plugins/Media/RTSPMedia/Source/RTSPMedia/RTSPMedia.Build.cs
+++ b/Engine/Plugins/Media/RTSPMedia/Source/RTSPMedia/RTSPMedia.Build.cs
@@ -36,5 +36,9 @@ public class RTSPMedia : ModuleRules
 			"TimeManagement"
 		});
 
+		if (Target.IsInPlatformGroup(UnrealPlatformGroup.Apple))
+		{
+			PublicFrameworks.Add("CoreVideo");
+		}
 	}
 }
\ No newline at end of file
