[RTSPMedia] Enable bProvideCpuBuffer by default

#rnx
#Synced-CL 52128593

diff --git a/Engine/Plugins/Media/RTSPMedia/Source/RTSPMedia/Public/RtspMediaDefaults.h b/Engine/Plugins/Media/RTSPMedia/Source/RTSPMedia/Public/RtspMediaDefaults.h
--- a/Engine/Plugins/Media/RTSPMedia/Source/RTSPMedia/Public/RtspMediaDefaults.h
+++ b/Engine/Plugins/Media/RTSPMedia/Source/RTSPMedia/Public/RtspMediaDefaults.h
@@ -7,7 +7,7 @@
 namespace RtspMedia::Default
 {
 	inline constexpr bool bAutoReconnect = true;
-	inline constexpr bool bProvideCpuBuffer = false;
+	inline constexpr bool bProvideCpuBuffer = true;
 	inline constexpr int32 DecoderBufferSize = 5;
 	inline constexpr int32 DecoderPollIntervalMs = 5;
 	inline constexpr bool bJitterBufferAutoAdjust = true;
